#ue4
#ios
* fix for looking for the All_iOS_On device which doesn't exist when wanting to deploy to all devices.

diff --git a/Engine/Source/Programs/IOS/iPhonePackager/DeployTime.cs b/Engine/Source/Programs/IOS/iPhonePackager/DeployTime.cs
--- a/Engine/Source/Programs/IOS/iPhonePackager/DeployTime.cs
+++ b/Engine/Source/Programs/IOS/iPhonePackager/DeployTime.cs
@@ -139,7 +139,7 @@ namespace iPhonePackager
 
 						if (File.Exists(GamePath))
 						{
-							if (!String.IsNullOrEmpty(Config.DeviceId))
+							if (!String.IsNullOrEmpty(Config.DeviceId) && !Config.DeviceId.Contains("All_iOS_On"))
 							{
 								DeploymentHelper.Get().DeviceId = Config.DeviceId;
 							}
