Force exact match when excluding Info.plist

This is so we don't accidentally exclude Info.plist from bundles or other sub-directories

diff --git a/Engine/Source/Programs/IOS/iPhonePackager/CodeSigning.cs b/Engine/Source/Programs/IOS/iPhonePackager/CodeSigning.cs
--- a/Engine/Source/Programs/IOS/iPhonePackager/CodeSigning.cs
+++ b/Engine/Source/Programs/IOS/iPhonePackager/CodeSigning.cs
@@ -240,7 +240,7 @@ namespace iPhonePackager
 			// Create a rules dict that includes (by wildcard) everything but Info.plist and the rules file
 			Dictionary<string, object> Rules = new Dictionary<string, object>();
 			Rules.Add(".*", true);
-			Rules.Add("Info.plist", CreateOmittedResource(10));
+			Rules.Add("^Info.plist", CreateOmittedResource(10));
 			Rules.Add(CFBundleResourceSpecification, CreateOmittedResource(100));
 
 			// Write the rules file out 
