Adding some missing files to the 4.2 install for iOS content-only projects on PC
#ue4
#ios

diff --git a/Engine/Source/Programs/AutomationTool/IOS/IOSPlatform.Automation.cs b/Engine/Source/Programs/AutomationTool/IOS/IOSPlatform.Automation.cs
--- a/Engine/Source/Programs/AutomationTool/IOS/IOSPlatform.Automation.cs
+++ b/Engine/Source/Programs/AutomationTool/IOS/IOSPlatform.Automation.cs
@@ -522,7 +522,7 @@ public class IOSPlatform : Platform
 					SourceProvision = GameSourceProvision;
 				}
 
-				SC.StageFiles(StagedFileType.NonUFS, Path.GetDirectoryName(SourceProvision), Path.GetFileName(SourceProvision), false, null, "", false, false, "embedded.mobileprovision");
+				SC.StageFiles(StagedFileType.NonUFS, Path.GetDirectoryName(SourceProvision), Path.GetFileName(SourceProvision), false, null, "", GlobalCommandLine.Rocket, false, "embedded.mobileprovision");
 			}
 		}
 	}
diff --git a/Engine/Source/Programs/IOS/iPhonePackager/iPhonePackager.cs b/Engine/Source/Programs/IOS/iPhonePackager/iPhonePackager.cs
--- a/Engine/Source/Programs/IOS/iPhonePackager/iPhonePackager.cs
+++ b/Engine/Source/Programs/IOS/iPhonePackager/iPhonePackager.cs
@@ -395,7 +395,7 @@ namespace iPhonePackager
 			}
 			else if (AssemblyName.ToLowerInvariant() == "ionic.zip.reduced")
 			{
-				AssemblyName = Path.GetFullPath(@"..\..\ThirdParty\Ionic\" + AssemblyName + ".dll");
+				AssemblyName = Path.GetFullPath(@"..\" + AssemblyName + ".dll");
 
 				Debug.WriteLineIf(System.Diagnostics.Debugger.IsAttached, "Loading assembly: " + AssemblyName);
 
@@ -496,6 +496,7 @@ namespace iPhonePackager
 				}
 
 				Log("Executing iPhonePackager " + String.Join(" ", args));
+				Log("CWD: " + Directory.GetCurrentDirectory());
 
 				// Ensure shipping configuration for final distributions
 				if (Config.bForDistribution && (GameConfiguration != "Shipping"))
