UDKRemote compile error fixes.
Missing resource and updated the deprecated usage of UUID

diff --git a/Engine/Source/Programs/IOS/UDKRemote/Classes/IPhoneHome.mm b/Engine/Source/Programs/IOS/UDKRemote/Classes/IPhoneHome.mm
--- a/Engine/Source/Programs/IOS/UDKRemote/Classes/IPhoneHome.mm
+++ b/Engine/Source/Programs/IOS/UDKRemote/Classes/IPhoneHome.mm
@@ -344,10 +344,7 @@ - (void)gameThreadCollectStats
 - (void)getUUID
 {
 	// get the UUID
-	NSData* Data = [[[UIDevice currentDevice] uniqueIdentifier] dataUsingEncoding:NSUTF8StringEncoding];
-	
-	// convert it to a string
-	UniqueId = DataToString(Data, YES);
+	UniqueId = [[NSUUID UUID] UUIDString];
 	[UniqueId retain];
 }
 
