Merged Github pull request #64 - DotNET security warnings when using iPhonePackager on Windows 8
#ue4
#ios

diff --git a/Engine/Source/Programs/IOS/iPhonePackager/Properties/iPhonePackager.exe.config b/Engine/Source/Programs/IOS/iPhonePackager/Properties/iPhonePackager.exe.config
--- a/Engine/Source/Programs/IOS/iPhonePackager/Properties/iPhonePackager.exe.config
+++ b/Engine/Source/Programs/IOS/iPhonePackager/Properties/iPhonePackager.exe.config
@@ -2,5 +2,6 @@
 <configuration>
 	<runtime>
 		<generatePublisherEvidence enabled="false"/>
+    <loadFromRemoteSources enabled="true"/>
 	</runtime>
 </configuration>
diff --git a/Engine/Source/Programs/IOS/iPhonePackager/app.config b/Engine/Source/Programs/IOS/iPhonePackager/app.config
--- a/Engine/Source/Programs/IOS/iPhonePackager/app.config
+++ b/Engine/Source/Programs/IOS/iPhonePackager/app.config
@@ -5,5 +5,6 @@
 	</startup>
 	<runtime>
 		<generatePublisherEvidence enabled="false"/>
+    <loadFromRemoteSources enabled="true"/>
 	</runtime>
 </configuration>
