Fix for using a space as the delimiter between two fully path files which could potentially include a space in the path when creating commands to upload files to the remote mac
#ue4
#ios

diff --git a/Engine/Source/Programs/IOS/iPhonePackager/FileOperations.cs b/Engine/Source/Programs/IOS/iPhonePackager/FileOperations.cs
--- a/Engine/Source/Programs/IOS/iPhonePackager/FileOperations.cs
+++ b/Engine/Source/Programs/IOS/iPhonePackager/FileOperations.cs
@@ -241,7 +241,7 @@ namespace iPhonePackager
 			foreach (FileInfo SourceFileInfo in SourceFolderInfo.GetFiles())
 			{
 				string DestFileName = DestFolder + "/" + SourceFileInfo.Name;
-				FilesToUpload.Add(SourceFileInfo.FullName + " " + DestFileName);
+				FilesToUpload.Add(SourceFileInfo.FullName + ";" + DestFileName);
 			}
 
 			// recurse into subdirectories
diff --git a/Engine/Source/Programs/RPCUtility/Program.cs b/Engine/Source/Programs/RPCUtility/Program.cs
--- a/Engine/Source/Programs/RPCUtility/Program.cs
+++ b/Engine/Source/Programs/RPCUtility/Program.cs
@@ -300,7 +300,7 @@ namespace RPCUtility
 			foreach (string Pair in CopyCommands)
 			{
 				// each pair is a local and remote filename
-				string[] Filenames = Pair.Split(" ".ToCharArray());
+				string[] Filenames = Pair.Split(";".ToCharArray());
 				if (Filenames.Length == 2)
 				{
 					FileInfo Info = new FileInfo(Filenames[0]);
diff --git a/Engine/Source/Programs/UnrealBuildTool/ToolChain/RemoteToolChain.cs b/Engine/Source/Programs/UnrealBuildTool/ToolChain/RemoteToolChain.cs
--- a/Engine/Source/Programs/UnrealBuildTool/ToolChain/RemoteToolChain.cs
+++ b/Engine/Source/Programs/UnrealBuildTool/ToolChain/RemoteToolChain.cs
@@ -359,7 +359,7 @@ namespace UnrealBuildTool
 						{
 							string RemoteFilePath = ConvertPath(SyncFile);
 							// an upload command is local name and remote name
-							BatchUploadCommands.Add(SyncFile + " " + RemoteFilePath);
+							BatchUploadCommands.Add(SyncFile + ";" + RemoteFilePath);
 						}
 					}
 				}
