more space in path fixes for IPP
#ue4
#ios

diff --git a/Engine/Source/Programs/IOS/iPhonePackager/CompileTime.cs b/Engine/Source/Programs/IOS/iPhonePackager/CompileTime.cs
--- a/Engine/Source/Programs/IOS/iPhonePackager/CompileTime.cs
+++ b/Engine/Source/Programs/IOS/iPhonePackager/CompileTime.cs
@@ -78,7 +78,7 @@ namespace iPhonePackager
 			// sign with the Distribution identity when packaging for distribution
 			if (!Config.bForDistribution)
 			{
-				CmdLine += String.Format(" CODE_SIGN_RESOURCE_RULES_PATH={0}/CustomResourceRules.plist", MacStagingRootDir);
+				CmdLine += String.Format(" CODE_SIGN_RESOURCE_RULES_PATH='{0}/CustomResourceRules.plist'", MacStagingRootDir);
 			}
 			CmdLine += String.Format(" CODE_SIGN_IDENTITY=\\\"{0}\\\"", Config.CodeSigningIdentity);
 
@@ -248,7 +248,7 @@ namespace iPhonePackager
 			case "deletemacstagingfiles":
 				Program.Log( " ... deleting staging files on the Mac" );
 				DisplayCommandLine = "rm -rf Payload";
-				CommandLine = MacStagingRootDir + " " + DisplayCommandLine;
+				CommandLine = "\"" + MacStagingRootDir + "\" " + DisplayCommandLine;
 				WorkingFolder = MacStagingRootDir;
 				break;
 
@@ -257,7 +257,7 @@ namespace iPhonePackager
 
 				DisplayCommandLine = String.Format("mkdir -p ~/Library/MobileDevice/Provisioning\\ Profiles");
 
-				CommandLine = MacXcodeStagingDir + " " + DisplayCommandLine;
+				CommandLine = "\"" + MacXcodeStagingDir + "\" " + DisplayCommandLine;
 				WorkingFolder = MacXcodeStagingDir;
 				break;
 
@@ -267,29 +267,29 @@ namespace iPhonePackager
 
 				DisplayCommandLine = String.Format("cp -f {0} ~/Library/MobileDevice/Provisioning\\ Profiles", MacMobileProvisionFilename);
 
-				CommandLine = MacXcodeStagingDir + " " + DisplayCommandLine;
+				CommandLine = "\"" + MacXcodeStagingDir + "\" " + DisplayCommandLine;
 				WorkingFolder = MacXcodeStagingDir;
 				break;
 
 			case "removeprovision":
 				Program.Log(" ... removing .mobileprovision");
 				DisplayCommandLine = String.Format("rm -f ~/Library/MobileDevice/Provisioning\\ Profiles/{0}", MacMobileProvisionFilename);
-				CommandLine = MacXcodeStagingDir + " " + DisplayCommandLine;
+				CommandLine = "\"" + MacXcodeStagingDir + "\" " + DisplayCommandLine;
 				WorkingFolder = MacXcodeStagingDir;
 				break;
 
 			case "setexec": 
 				// Note: The executable must have already been copied over
 				Program.Log(" ... setting executable bit");
-				DisplayCommandLine = "chmod a+x " + RemoteExecutablePath;
-				CommandLine = MacStagingRootDir + " " + DisplayCommandLine;
+				DisplayCommandLine = "chmod a+x \'" + RemoteExecutablePath + "\'";
+				CommandLine = "\"" + MacStagingRootDir + "\" " + DisplayCommandLine;
 				WorkingFolder = MacStagingRootDir;
 				break;
 
 			case "prepackage":
 				Program.Log(" ... running prepackage script remotely ");
 				DisplayCommandLine = String.Format("sh prepackage.sh {0} IOS {1} {2}", Program.GameName, Program.GameConfiguration, Program.Architecture);
-				CommandLine = MacXcodeStagingDir + " " + DisplayCommandLine;
+				CommandLine = "\"" + MacXcodeStagingDir + "\" " + DisplayCommandLine;
 				WorkingFolder = MacXcodeStagingDir;
 				break;
 
@@ -297,21 +297,21 @@ namespace iPhonePackager
 				Program.Log(" ... making application (codesign, etc...)");
 				Program.Log("  Using signing identity '{0}'", Config.CodeSigningIdentity);
 				DisplayCommandLine = CurrentBaseXCodeCommandLine;
-				CommandLine = MacXcodeStagingDir + "/.. " + DisplayCommandLine;
-				WorkingFolder = MacXcodeStagingDir + "/..";
+				CommandLine = "\"" + MacXcodeStagingDir + "/..\" " + DisplayCommandLine;
+				WorkingFolder = "\"" + MacXcodeStagingDir + "/..\"";
 				break;
 
 			case "validation":
 				Program.Log( " ... validating distribution package" );
 				DisplayCommandLine = XcodeDeveloperDir + "Platforms/iPhoneOS.platform/Developer/usr/bin/Validation " + RemoteAppDirectory;
-				CommandLine = MacStagingRootDir + " " + DisplayCommandLine;
+				CommandLine = "\"" + MacStagingRootDir + "\" " + DisplayCommandLine;
 				WorkingFolder = MacStagingRootDir;
 				break;
 
 			case "deleteipa":
 				Program.Log(" ... deleting IPA on Mac");
 				DisplayCommandLine = "rm -f " + Config.IPAFilenameOnMac;
-				CommandLine = MacStagingRootDir + " " + DisplayCommandLine;
+				CommandLine = "\"" + MacStagingRootDir + "\" " + DisplayCommandLine;
 				WorkingFolder = MacStagingRootDir;
 				break;
 
@@ -325,14 +325,14 @@ namespace iPhonePackager
 			case "strip":
 				Program.Log( " ... stripping" );
 				DisplayCommandLine = XcodeDeveloperDir + "Platforms/iPhoneOS.platform/Developer/usr/bin/strip " + RemoteExecutablePath;
-				CommandLine = MacStagingRootDir + " " + DisplayCommandLine;
+				CommandLine = "\"" + MacStagingRootDir + "\" " + DisplayCommandLine;
 				WorkingFolder = MacStagingRootDir;
 				break;
 
 			case "resign":
 				Program.Log("... resigning");
 				DisplayCommandLine = "bash -c '" + "chmod a+x ResignScript" + ";" + "./ResignScript" + "'";
-				CommandLine = MacStagingRootDir + " " + DisplayCommandLine;
+				CommandLine = "\"" + MacStagingRootDir + "\" " + DisplayCommandLine;
 				WorkingFolder = MacStagingRootDir;
 				break;
 
@@ -349,7 +349,7 @@ namespace iPhonePackager
 					Config.IPAFilenameOnMac,
 					dSYMName);
 
-				CommandLine = MacStagingRootDir + " " + DisplayCommandLine;
+				CommandLine = "\"" + MacStagingRootDir + "\" " + DisplayCommandLine;
 				WorkingFolder = MacStagingRootDir;
 				break;
 
@@ -360,7 +360,7 @@ namespace iPhonePackager
 				string dSYMPath = Program.GameName + ".app.dSYM";
 				DisplayCommandLine = String.Format("dsymutil -o {0} {1}", dSYMPath, ExePath);
 
-				CommandLine = MacStagingRootDir + " " + DisplayCommandLine;
+				CommandLine = "\"" + MacStagingRootDir + "\"" + DisplayCommandLine;
 				WorkingFolder = MacStagingRootDir;
 				break;
 
@@ -371,6 +371,7 @@ namespace iPhonePackager
 
 			Program.Log( " ... working folder: " + WorkingFolder );
 			Program.Log( " ... " + DisplayCommandLine );
+			Program.Log(" ... full command: " +  MacName + " " + CommandLine);
 
 			Process RPCUtil = new Process();
 			RPCUtil.StartInfo.FileName = @"..\RPCUtility.exe";
