more remote tool chain fixes for IOS
fix for spaces in paths for repackaging ipa and creating stub
fix for deployment of the support files for code based projects
#ue4
#ios

diff --git a/Engine/Source/Programs/AutomationTool/IOS/IOSPlatform.Automation.cs b/Engine/Source/Programs/AutomationTool/IOS/IOSPlatform.Automation.cs
--- a/Engine/Source/Programs/AutomationTool/IOS/IOSPlatform.Automation.cs
+++ b/Engine/Source/Programs/AutomationTool/IOS/IOSPlatform.Automation.cs
@@ -124,7 +124,7 @@ public class IOSPlatform : Platform
 
 			bool cookonthefly = Params.CookOnTheFly || Params.SkipCookOnTheFly;
 
-			string IPPArguments = "RepackageFromStage " + (Params.IsCodeBasedProject ? Params.RawProjectPath : "Engine");
+			string IPPArguments = "RepackageFromStage \"" + (Params.IsCodeBasedProject ? Params.RawProjectPath : "Engine") + "\"";
 			IPPArguments += " -config " + TargetConfiguration.ToString();
 
 			if (TargetConfiguration == UnrealTargetConfiguration.Shipping)
diff --git a/Engine/Source/Programs/IOS/iPhonePackager/CompileTime.cs b/Engine/Source/Programs/IOS/iPhonePackager/CompileTime.cs
--- a/Engine/Source/Programs/IOS/iPhonePackager/CompileTime.cs
+++ b/Engine/Source/Programs/IOS/iPhonePackager/CompileTime.cs
@@ -126,8 +126,8 @@ namespace iPhonePackager
 
 			// generate the directories to recursively copy into later on
 			MacStagingRootDir = string.Format("{0}/{1}/IOS", iPhone_SigningDevRootMac, GameBranchPath);
-			MacBinariesDir = string.Format("{0}/{1}/IOS", iPhone_SigningDevRootMac, BranchPath);
-			MacXcodeStagingDir = string.Format("{0}/{1}/IOS/XcodeSupportFiles", iPhone_SigningDevRootMac, BranchPath);
+			MacBinariesDir = string.Format("{0}/{1}/IOS", iPhone_SigningDevRootMac, GameBranchPath);
+			MacXcodeStagingDir = string.Format("{0}/{1}/IOS/XcodeSupportFiles", iPhone_SigningDevRootMac, GameBranchPath);
 
 			MacMobileProvisionFilename = MachineName + "_UE4Temp.mobileprovision";
 
diff --git a/Engine/Source/Programs/UnrealBuildTool/IOS/IOSToolChain.cs b/Engine/Source/Programs/UnrealBuildTool/IOS/IOSToolChain.cs
--- a/Engine/Source/Programs/UnrealBuildTool/IOS/IOSToolChain.cs
+++ b/Engine/Source/Programs/UnrealBuildTool/IOS/IOSToolChain.cs
@@ -1186,7 +1186,7 @@ namespace UnrealBuildTool
 					}
 					else
 					{
-						StubGenerateProcess.StartInfo.Arguments = "PackageIPA " + PathToApp + " -createstub";
+						StubGenerateProcess.StartInfo.Arguments = "PackageIPA \"" + PathToApp + "\" -createstub";
 						// if we are making the dsym, then we can strip the debug info from the executable
 						if (BuildConfiguration.bStripSymbolsOnIOS || (Target.Configuration == UnrealTargetConfiguration.Shipping))
 						{
