moved to a library based approach for mobile provisions on PC
TTP326334
#ue4
#ios

diff --git a/Engine/Source/Programs/AutomationTool/IOS/IOSPlatform.Automation.cs b/Engine/Source/Programs/AutomationTool/IOS/IOSPlatform.Automation.cs
--- a/Engine/Source/Programs/AutomationTool/IOS/IOSPlatform.Automation.cs
+++ b/Engine/Source/Programs/AutomationTool/IOS/IOSPlatform.Automation.cs
@@ -538,32 +538,6 @@ public class IOSPlatform : Platform
 
 				SC.StageFiles(StagedFileType.NonUFS, SourcePath, Path.GetFileName(TargetPListFile), false, null, "", false, false, "Info.plist");
 			}
-
-			// Now do the .mobileprovision
-			//@TODO: Remove this mobileprovision copy, and move to a library approach like Xcode/codesign does
-			if (GetCodeSignDesirability(Params))
-			{
-				string SourceProvision = CombinePaths(SC.LocalRoot, "Engine", "Build", "IOS", "UE4Game.mobileprovision");
-				string GameSourceProvision = CombinePaths(SC.ProjectRoot, "Build", "IOS", SC.ShortProjectName + ".mobileprovision");
-				if (!File.Exists(GameSourceProvision))
-				{
-					GameSourceProvision = CombinePaths(SC.ProjectRoot, "Build", "IOS", "NotForLicensees", SC.ShortProjectName + ".mobileprovision");
-					if (File.Exists(GameSourceProvision))
-					{
-						SourceProvision = GameSourceProvision;
-					}
-					else if (!File.Exists(SourceProvision))
-					{
-						SourceProvision = CombinePaths(SC.LocalRoot, "Engine", "Build", "IOS", "NotForLicensees", "UE4Game.mobileprovision");
-					}
-				}
-				else
-				{
-					SourceProvision = GameSourceProvision;
-				}
-
-				SC.StageFiles(StagedFileType.NonUFS, Path.GetDirectoryName(SourceProvision), Path.GetFileName(SourceProvision), false, null, "", GlobalCommandLine.Rocket, false, "embedded.mobileprovision");
-			}
 		}
 
 		// copy the movies from the project
diff --git a/Engine/Source/Programs/IOS/iPhonePackager/CodeSigning.cs b/Engine/Source/Programs/IOS/iPhonePackager/CodeSigning.cs
--- a/Engine/Source/Programs/IOS/iPhonePackager/CodeSigning.cs
+++ b/Engine/Source/Programs/IOS/iPhonePackager/CodeSigning.cs
@@ -76,7 +76,7 @@ namespace iPhonePackager
 			Provision = null;
 			try
 			{
-				string ExpectedProvisionFile = FileOperations.FindPrefixedFile(Config.BuildDirectory, Program.GameName + ".mobileprovision");
+				string ExpectedProvisionFile = FileOperations.FindPrefixedFile(Config.ProvisionDirectory, Program.GameName + ".mobileprovision");
 				Provision = MobileProvisionParser.ParseFile(ExpectedProvisionFile);
 			}
 			catch (Exception)
@@ -105,30 +105,8 @@ namespace iPhonePackager
 
 		protected virtual byte[] GetMobileProvision()
 		{
-			// Find the mobile provision file to use
-			string MobileProvisionFilename = Path.Combine(Config.RepackageStagingDirectory, "embedded.mobileprovision");
-				
-			if (!File.Exists(MobileProvisionFilename))
-			{
-				MobileProvisionFilename = FileOperations.FindPrefixedFile(Config.BuildDirectory, Program.GameName + ".mobileprovision");
-			}
-
-			if (!File.Exists(MobileProvisionFilename))
-			{
-				MobileProvisionFilename = FileOperations.FindPrefixedFile(Config.BuildDirectory + "/NotForLicensees/", Program.GameName + ".mobileprovision");
-				if (!File.Exists(MobileProvisionFilename))
-				{
-					MobileProvisionFilename = FileOperations.FindPrefixedFile(Config.EngineBuildDirectory, "UE4Game.mobileprovision");
-					if (!File.Exists(MobileProvisionFilename))
-					{
-						MobileProvisionFilename = FileOperations.FindPrefixedFile(Config.EngineBuildDirectory + "/NotForLicensees/", "UE4Game.mobileprovision");
-						if (!File.Exists(MobileProvisionFilename))
-						{
-							MobileProvisionFilename = FileOperations.FindAnyFileWithExtension(Config.BuildDirectory, ".mobileprovision");
-						}
-					}
-				}
-			}
+			// find the movile provision file in the library
+			string MobileProvisionFilename = MobileProvision.FindCompatibleProvision();
 
 			byte[] Result = null;
 			try
diff --git a/Engine/Source/Programs/IOS/iPhonePackager/Config.cs b/Engine/Source/Programs/IOS/iPhonePackager/Config.cs
--- a/Engine/Source/Programs/IOS/iPhonePackager/Config.cs
+++ b/Engine/Source/Programs/IOS/iPhonePackager/Config.cs
@@ -64,6 +64,14 @@ namespace iPhonePackager
 			}
 		}
 
+		/// <summary>
+		/// The shared provision library directory (on PC)
+		/// </summary>
+		public static string ProvisionDirectory
+		{
+			get { return Environment.GetFolderPath(Environment.SpecialFolder.LocalApplicationData) + "/Apple Computer/MobileDevice/Provisioning Profiles/"; }
+		}
+
 		/// <summary>
 		/// The shared (Engine) build directory (on PC)
 		/// </summary>
diff --git a/Engine/Source/Programs/IOS/iPhonePackager/MobileProvisionUtilities.cs b/Engine/Source/Programs/IOS/iPhonePackager/MobileProvisionUtilities.cs
--- a/Engine/Source/Programs/IOS/iPhonePackager/MobileProvisionUtilities.cs
+++ b/Engine/Source/Programs/IOS/iPhonePackager/MobileProvisionUtilities.cs
@@ -21,11 +21,200 @@ namespace iPhonePackager
 		public object Tag;
 
 		public string ApplicationIdentifierPrefix = null;
+		public string ApplicationIdentifier = null;
 		public List<X509Certificate2> DeveloperCertificates = new List<X509Certificate2>();
 		public List<string> ProvisionedDeviceIDs;
 		public string ProvisionName;
+		public bool bDebug;
 		public Utilities.PListHelper Data;
 
+		public static string FindCompatibleProvision()
+		{
+			// remap the gamename if necessary
+			string GameName = Program.GameName;
+			if (GameName == "UE4Game")
+			{
+				string[] project = Directory.GetFiles(Config.ProjectRootDirectory, "*.uproject");
+				if (project.Length == 1)
+				{
+					GameName = Path.GetFileNameWithoutExtension(project[0]);
+				}
+			}
+
+			// cache the provision library
+			Dictionary<string, MobileProvision> ProvisionLibrary = new Dictionary<string, MobileProvision>();
+			foreach (string Provision in Directory.EnumerateFiles(Config.ProvisionDirectory, "*.mobileprovision"))
+			{
+				MobileProvision p = MobileProvisionParser.ParseFile(Provision);
+				ProvisionLibrary.Add(Provision, p);
+			}
+
+			// check the cache for a provision matching the app id (com.company.Game)
+			foreach (KeyValuePair<string, MobileProvision> Pair in ProvisionLibrary)
+			{
+				if ((Pair.Value.ProvisionName == GameName || Pair.Value.ApplicationIdentifier.Contains(GameName)) && (!Config.bForDistribution || (Pair.Value.ProvisionedDeviceIDs.Count == 0 && !Pair.Value.bDebug)))
+				{
+					// check to see if we have a certificate for this provision
+					if (CodeSignatureBuilder.FindCertificate(Pair.Value) != null)
+					{
+						return Pair.Key;
+					}
+				}
+			}
+
+			#region remove after we provide an install mechanism
+			// check to see if we can find a provision in the game's build directory
+			string MobileProvisionFilename = FileOperations.FindPrefixedFile(Config.ProjectRootDirectory + "/Build/IOS/", GameName + ".mobileprovision");
+
+			if (!File.Exists(MobileProvisionFilename))
+			{
+				MobileProvisionFilename = FileOperations.FindPrefixedFile(Config.ProjectRootDirectory + "/Build/IOS/NotForLicensees/", GameName + ".mobileprovision");
+				if (!File.Exists(MobileProvisionFilename))
+				{
+					MobileProvisionFilename = FileOperations.FindAnyFileWithExtension(Config.ProjectRootDirectory + "/Build/IOS/", ".mobileprovision");
+				}
+			}
+			if (File.Exists(MobileProvisionFilename))
+			{
+				Directory.CreateDirectory(Config.ProvisionDirectory);
+				File.Copy(MobileProvisionFilename, Config.ProvisionDirectory + Path.GetFileName(MobileProvisionFilename), true);
+				FileInfo DestFileInfo = new FileInfo(Config.ProvisionDirectory + Path.GetFileName(MobileProvisionFilename));
+				DestFileInfo.Attributes = DestFileInfo.Attributes & ~FileAttributes.ReadOnly;
+				MobileProvisionFilename = Config.ProvisionDirectory + Path.GetFileName(MobileProvisionFilename);
+
+				// verify the provision is good
+				MobileProvision p = MobileProvisionParser.ParseFile(MobileProvisionFilename);
+				if ((p.ProvisionName == GameName || p.ApplicationIdentifier.Contains(GameName)) && (!Config.bForDistribution || (p.ProvisionedDeviceIDs.Count == 0 && !p.bDebug)))
+				{
+					// check to see if we have a certificate for this provision
+					if (CodeSignatureBuilder.FindCertificate(p) != null)
+					{
+						return MobileProvisionFilename;
+					}
+				}
+			}
+
+			// check for a distribution provision in the game's build directory
+			if (Config.bForDistribution)
+			{
+				MobileProvisionFilename = FileOperations.FindPrefixedFile(Config.ProjectRootDirectory + "/Build/IOS/", GameName + "_Distro.mobileprovision");
+
+				if (!File.Exists(MobileProvisionFilename))
+				{
+					MobileProvisionFilename = FileOperations.FindPrefixedFile(Config.ProjectRootDirectory + "/Build/IOS/NotForLicensees/", GameName + "_Distro.mobileprovision");
+					if (!File.Exists(MobileProvisionFilename))
+					{
+						MobileProvisionFilename = FileOperations.FindAnyFileWithExtension(Config.ProjectRootDirectory + "/Build/IOS/", "_Distro.mobileprovision");
+					}
+				}
+				if (File.Exists(MobileProvisionFilename))
+				{
+					Directory.CreateDirectory(Config.ProvisionDirectory);
+					File.Copy(MobileProvisionFilename, Config.ProvisionDirectory + Path.GetFileName(MobileProvisionFilename), true);
+					FileInfo DestFileInfo = new FileInfo(Config.ProvisionDirectory + Path.GetFileName(MobileProvisionFilename));
+					DestFileInfo.Attributes = DestFileInfo.Attributes & ~FileAttributes.ReadOnly;
+					MobileProvisionFilename = Config.ProvisionDirectory + Path.GetFileName(MobileProvisionFilename);
+
+					// verify the provision is good
+					MobileProvision p = MobileProvisionParser.ParseFile(MobileProvisionFilename);
+					if ((p.ProvisionName == GameName || p.ApplicationIdentifier.Contains(GameName)) && (!Config.bForDistribution || (p.ProvisionedDeviceIDs.Count == 0 && !p.bDebug)))
+					{
+						// check to see if we have a certificate for this provision
+						if (CodeSignatureBuilder.FindCertificate(p) != null)
+						{
+							return MobileProvisionFilename;
+						}
+					}
+				}
+			}
+			#endregion
+
+			// check the cache for a provision matching the wild card
+			foreach (KeyValuePair<string, MobileProvision> Pair in ProvisionLibrary)
+			{
+				if ((Pair.Value.ProvisionName.Contains("Wildcard") || Pair.Value.ApplicationIdentifier.Contains("*")) && (!Config.bForDistribution || (Pair.Value.ProvisionedDeviceIDs.Count == 0 && !Pair.Value.bDebug)))
+				{
+					// check to see if we have a certificate for this provision
+					if (CodeSignatureBuilder.FindCertificate(Pair.Value) != null)
+					{
+						return Pair.Key;
+					}
+				}
+			}
+
+			#region remove after we provide an install mechanism
+			// check to see if we can find a provision in the engine's build directory
+			MobileProvisionFilename = FileOperations.FindPrefixedFile(Config.EngineBuildDirectory, "UE4Game.mobileprovision");
+
+			if (!File.Exists(MobileProvisionFilename))
+			{
+				MobileProvisionFilename = FileOperations.FindPrefixedFile(Config.EngineBuildDirectory + "/NotForLicensees/", "UE4Game.mobileprovision");
+				if (!File.Exists(MobileProvisionFilename))
+				{
+					MobileProvisionFilename = FileOperations.FindAnyFileWithExtension(Config.EngineBuildDirectory, ".mobileprovision");
+				}
+			}
+			if (File.Exists(MobileProvisionFilename))
+			{
+				Directory.CreateDirectory(Config.ProvisionDirectory);
+				File.Copy(MobileProvisionFilename, Config.ProvisionDirectory + Path.GetFileName(MobileProvisionFilename), true);
+				FileInfo DestFileInfo = new FileInfo(Config.ProvisionDirectory + Path.GetFileName(MobileProvisionFilename));
+				DestFileInfo.Attributes = DestFileInfo.Attributes & ~FileAttributes.ReadOnly;
+				MobileProvisionFilename = Config.ProvisionDirectory + Path.GetFileName(MobileProvisionFilename);
+
+				// verify the provision is good
+				MobileProvision p = MobileProvisionParser.ParseFile(MobileProvisionFilename);
+				if ((p.ProvisionName.Contains("Wildcard") || p.ApplicationIdentifier.Contains("*")) && (!Config.bForDistribution || (p.ProvisionedDeviceIDs.Count == 0 && !p.bDebug)))
+				{
+					// check to see if we have a certificate for this provision
+					if (CodeSignatureBuilder.FindCertificate(p) != null)
+					{
+						return MobileProvisionFilename;
+					}
+				}
+			}
+
+			// check for a distribution provision in the game's build directory
+			if (Config.bForDistribution)
+			{
+				MobileProvisionFilename = FileOperations.FindPrefixedFile(Config.EngineBuildDirectory, "UE4Game_Distro.mobileprovision");
+
+				if (!File.Exists(MobileProvisionFilename))
+				{
+					MobileProvisionFilename = FileOperations.FindPrefixedFile(Config.EngineBuildDirectory + "/NotForLicensees/", "UE4Game_Distro.mobileprovision");
+					if (!File.Exists(MobileProvisionFilename))
+					{
+						MobileProvisionFilename = FileOperations.FindAnyFileWithExtension(Config.EngineBuildDirectory, ".mobileprovision");
+					}
+				}
+				if (File.Exists(MobileProvisionFilename))
+				{
+					Directory.CreateDirectory(Config.ProvisionDirectory);
+					File.Copy(MobileProvisionFilename, Config.ProvisionDirectory + Path.GetFileName(MobileProvisionFilename), true);
+					FileInfo DestFileInfo = new FileInfo(Config.ProvisionDirectory + Path.GetFileName(MobileProvisionFilename));
+					DestFileInfo.Attributes = DestFileInfo.Attributes & ~FileAttributes.ReadOnly;
+					MobileProvisionFilename = Config.ProvisionDirectory + Path.GetFileName(MobileProvisionFilename);
+
+					// verify the provision is good
+					MobileProvision p = MobileProvisionParser.ParseFile(MobileProvisionFilename);
+					if ((p.ProvisionName == GameName || p.ApplicationIdentifier.Contains(GameName)) && (!Config.bForDistribution || (p.ProvisionedDeviceIDs.Count == 0 && !p.bDebug)))
+					{
+						// check to see if we have a certificate for this provision
+						if (CodeSignatureBuilder.FindCertificate(p) != null)
+						{
+							return MobileProvisionFilename;
+						}
+					}
+				}
+			}
+			#endregion
+
+			// check to see if there is already an embedded provision
+			MobileProvisionFilename = Path.Combine(Config.RepackageStagingDirectory, "embedded.mobileprovision");
+
+			return MobileProvisionFilename;
+		}
+
 		/// <summary>
 		/// Extracts the dict values for the Entitlements key and creates a new full .plist file
 		/// from them (with outer plist and dict keys as well as doctype, etc...)
@@ -109,6 +298,22 @@ namespace iPhonePackager
 
 			// Key: ProvisionedDevices, Array<String>
 			ProvisionedDeviceIDs = Data.GetArray("ProvisionedDevices", "string");
+
+			// Key: application-identifier, Array<String>
+			Utilities.PListHelper XCentPList = null;
+			Data.ProcessValueForKey("Entitlements", "dict", delegate(XmlNode ValueNode)
+			{
+				XCentPList = Utilities.PListHelper.CloneDictionaryRootedAt(ValueNode);
+			});
+
+			// Modify the application-identifier to be fully qualified if needed
+			if (!XCentPList.GetString("application-identifier", out ApplicationIdentifier))
+			{
+				ApplicationIdentifier = "(unknown)";
+			}
+
+			// check for get-task-allow
+			bDebug = XCentPList.GetBool("get-task-allow");
 		}
 
 		/// <summary>
diff --git a/Engine/Source/Programs/IOS/iPhonePackager/Utilities.cs b/Engine/Source/Programs/IOS/iPhonePackager/Utilities.cs
--- a/Engine/Source/Programs/IOS/iPhonePackager/Utilities.cs
+++ b/Engine/Source/Programs/IOS/iPhonePackager/Utilities.cs
@@ -351,6 +351,19 @@ namespace iPhonePackager
 				return true;
 			}
 
+			public bool GetBool(string Key)
+			{
+				string PathToValue = String.Format("/plist/dict/key[.='{0}']/following-sibling::node()", Key);
+
+				XmlNode ValueNode = Doc.DocumentElement.SelectSingleNode(PathToValue);
+				if (ValueNode == null)
+				{
+					return false;
+				}
+
+				return ValueNode.Name == "true";
+			}
+
 			public delegate void ProcessOneNodeEvent(XmlNode ValueNode);
 
 			public void ProcessValueForKey(string Key, string ExpectedValueType, ProcessOneNodeEvent ValueHandler)
diff --git a/Engine/Source/Programs/UnrealBuildTool/IOS/UEDeployIOS.cs b/Engine/Source/Programs/UnrealBuildTool/IOS/UEDeployIOS.cs
--- a/Engine/Source/Programs/UnrealBuildTool/IOS/UEDeployIOS.cs
+++ b/Engine/Source/Programs/UnrealBuildTool/IOS/UEDeployIOS.cs
@@ -212,10 +212,12 @@ namespace UnrealBuildTool.IOS
 
 		public override bool PrepTargetForDeployment(UEBuildTarget InTarget)
 		{
+			string GameName = InTarget.AppName;
+			string BuildPath = InTarget.ProjectDirectory + "/Binaries/IOS";
+			string ProjectDirectory = InTarget.ProjectDirectory;
+
 			if (ExternalExecution.GetRuntimePlatform() == UnrealTargetPlatform.Mac && Environment.GetEnvironmentVariable("UBT_NO_POST_DEPLOY") != "true")
 			{
-				string GameName = InTarget.AppName;
-
 				string DecoratedGameName;
 				if (InTarget.Configuration == UnrealTargetConfiguration.Development)
 				{
@@ -226,9 +228,6 @@ namespace UnrealBuildTool.IOS
 					DecoratedGameName = String.Format("{0}-{1}-{2}", GameName, InTarget.Platform.ToString(), InTarget.Configuration.ToString());
 				}
 
-				string BuildPath = InTarget.ProjectDirectory + "/Binaries/IOS";
-				string ProjectDirectory = InTarget.ProjectDirectory;
-
 				return PrepForUATPackageOrDeploy(GameName, ProjectDirectory, BuildPath + "/" + DecoratedGameName, "../../Engine", false, "");
 			}
 			else
@@ -270,6 +269,56 @@ namespace UnrealBuildTool.IOS
 						Log.TraceInformation("Copying binaries back to this device failed.");
 					}
 				}
+
+				// install the provision
+/*				string ProvisionWithPrefix = "../../Engine/Build/IOS/UE4Game.mobileprovision";
+				if (File.Exists(BuildPath + "/" + GameName + ".mobileprovision"))
+				{
+					ProvisionWithPrefix = BuildPath + "/" + GameName + ".mobileprovision";
+				}
+				else
+				{
+					if (File.Exists(BuildPath + "/NotForLicensees/" + GameName + ".mobileprovision"))
+					{
+						ProvisionWithPrefix = BuildPath + "/NotForLicensees/" + GameName + ".mobileprovision";
+					}
+					else if (!File.Exists(ProvisionWithPrefix))
+					{
+						ProvisionWithPrefix = "../../Engine/Build/IOS/NotForLicensees/UE4Game.mobileprovision";
+					}
+				}
+				string LibraryDir = Environment.GetFolderPath(Environment.SpecialFolder.LocalApplicationData) + "/Apple Computer/MobileDevice/Provisioning Profiles/";
+				if (File.Exists(ProvisionWithPrefix))
+				{
+					Directory.CreateDirectory(LibraryDir);
+					File.Copy(ProvisionWithPrefix, LibraryDir + GameName + ".mobileprovision", true);
+					FileInfo DestFileInfo = new FileInfo(LibraryDir + GameName + ".mobileprovision");
+					DestFileInfo.Attributes = DestFileInfo.Attributes & ~FileAttributes.ReadOnly;
+				}
+
+				// install the distribution provision
+				ProvisionWithPrefix = "../../Engine/Build/IOS/UE4Game_Distro.mobileprovision";
+				if (File.Exists(BuildPath + "/" + GameName + "_Distro.mobileprovision"))
+				{
+					ProvisionWithPrefix = BuildPath + "/" + GameName + "_Distro.mobileprovision";
+				}
+				else
+				{
+					if (File.Exists(BuildPath + "/NotForLicensees/" + GameName + "_Distro.mobileprovision"))
+					{
+						ProvisionWithPrefix = BuildPath + "/NotForLicensees/" + GameName + "_Distro.mobileprovision";
+					}
+					else if (!File.Exists(ProvisionWithPrefix))
+					{
+						ProvisionWithPrefix = "../../Engine/Build/IOS/NotForLicensees/UE4Game_Distro.mobileprovision";
+					}
+				}
+				if (File.Exists(ProvisionWithPrefix))
+				{
+					File.Copy(ProvisionWithPrefix, LibraryDir + GameName + "_Distro.mobileprovision", true);
+					FileInfo DestFileInfo = new FileInfo(LibraryDir + GameName + "_Distro.mobileprovision");
+					DestFileInfo.Attributes = DestFileInfo.Attributes & ~FileAttributes.ReadOnly;
+				}*/
 			}
 			return true;
 		}
