fix for not checking for a valid directory when trying to cache the mobile provisions
TTP339484
#ue4
#ios

diff --git a/Engine/Source/Programs/IOS/iPhonePackager/MobileProvisionUtilities.cs b/Engine/Source/Programs/IOS/iPhonePackager/MobileProvisionUtilities.cs
--- a/Engine/Source/Programs/IOS/iPhonePackager/MobileProvisionUtilities.cs
+++ b/Engine/Source/Programs/IOS/iPhonePackager/MobileProvisionUtilities.cs
@@ -41,6 +41,12 @@ namespace iPhonePackager
 				}
 			}
 
+			// ensure the provision directory exists
+			if (!Directory.Exists(Config.ProvisionDirectory))
+			{
+				Directory.CreateDirectory(Config.ProvisionDirectory);
+			}
+
 			// cache the provision library
 			Dictionary<string, MobileProvision> ProvisionLibrary = new Dictionary<string, MobileProvision>();
 			foreach (string Provision in Directory.EnumerateFiles(Config.ProvisionDirectory, "*.mobileprovision"))
