search for provisions by bundle identifier
better solution for copying legacy provisions
#ue4
#ios

diff --git a/Engine/Source/Programs/IOS/DeploymentServer/DeployTime.cs b/Engine/Source/Programs/IOS/DeploymentServer/DeployTime.cs
--- a/Engine/Source/Programs/IOS/DeploymentServer/DeployTime.cs
+++ b/Engine/Source/Programs/IOS/DeploymentServer/DeployTime.cs
@@ -307,6 +307,7 @@ namespace DeploymentServer
                     ReportIF.Warning(String.Format("Device '{0}' is a {1} model, which does not support OpenGL ES2.0.  The installation is likely to fail.", DeviceName, GetPrettyDeviceType(DeviceType)));
                 }
 
+				Console.WriteLine("Device '{0}' with id {1} is being checked against {2}.", Device.DeviceName, Device.DeviceId, DeviceId);
 				if (String.IsNullOrEmpty(DeviceId) || Device.DeviceId == DeviceId)
 				{
 					ReportIF.Log(String.Format("Transferring IPA to device '{0}' ... ", DeviceName));
diff --git a/Engine/Source/Programs/IOS/iPhonePackager/CodeSigning.cs b/Engine/Source/Programs/IOS/iPhonePackager/CodeSigning.cs
--- a/Engine/Source/Programs/IOS/iPhonePackager/CodeSigning.cs
+++ b/Engine/Source/Programs/IOS/iPhonePackager/CodeSigning.cs
@@ -103,10 +103,10 @@ namespace iPhonePackager
 			return bOverridesExists && (Provision != null) && (Cert != null);
 		}
 
-		protected virtual byte[] GetMobileProvision()
+		protected virtual byte[] GetMobileProvision(string CFBundleIdentifier)
 		{
 			// find the movile provision file in the library
-			string MobileProvisionFilename = MobileProvision.FindCompatibleProvision();
+			string MobileProvisionFilename = MobileProvision.FindCompatibleProvision(CFBundleIdentifier);
 
 			byte[] Result = null;
 			try
@@ -122,9 +122,9 @@ namespace iPhonePackager
 			return Result;
 		}
 
-		public void LoadMobileProvision()
+		public void LoadMobileProvision(string CFBundleIdentifier)
 		{
-			byte[] MobileProvisionFile = GetMobileProvision();
+			byte[] MobileProvisionFile = GetMobileProvision(CFBundleIdentifier);
 
 			if (MobileProvisionFile != null)
 			{
@@ -290,8 +290,15 @@ namespace iPhonePackager
 			// Load Info.plist, which guides nearly everything else
 			Info = LoadInfoPList();
 
+			// Get the name of the bundle
+			string CFBundleIdentifier;
+			if (!Info.GetString("CFBundleIdentifier", out CFBundleIdentifier))
+			{
+				throw new InvalidDataException("Info.plist must contain the key CFBundleIdentifier");
+			}
+
 			// Load the mobile provision, which provides entitlements and a partial cert which can be used to find an installed certificate
-			LoadMobileProvision();
+			LoadMobileProvision(CFBundleIdentifier);
 			if (Provision == null)
 			{
 				return;
diff --git a/Engine/Source/Programs/IOS/iPhonePackager/GraphicalResignTool.cs b/Engine/Source/Programs/IOS/iPhonePackager/GraphicalResignTool.cs
--- a/Engine/Source/Programs/IOS/iPhonePackager/GraphicalResignTool.cs
+++ b/Engine/Source/Programs/IOS/iPhonePackager/GraphicalResignTool.cs
@@ -247,7 +247,7 @@ namespace iPhonePackager
 			/// </summary>
 			public byte[] CustomInfoPList = null;
 
-			protected override byte[] GetMobileProvision()
+			protected override byte[] GetMobileProvision(string CFBundleIdentifier)
 			{
 				return (CustomMobileProvision != null) ? CustomMobileProvision : FileSystem.ReadAllBytes("embedded.mobileprovision");
 			}
diff --git a/Engine/Source/Programs/IOS/iPhonePackager/MobileProvisionUtilities.cs b/Engine/Source/Programs/IOS/iPhonePackager/MobileProvisionUtilities.cs
--- a/Engine/Source/Programs/IOS/iPhonePackager/MobileProvisionUtilities.cs
+++ b/Engine/Source/Programs/IOS/iPhonePackager/MobileProvisionUtilities.cs
@@ -28,7 +28,7 @@ namespace iPhonePackager
 		public bool bDebug;
 		public Utilities.PListHelper Data;
 
-		public static string FindCompatibleProvision()
+		public static string FindCompatibleProvision(string CFBundleIdentifier)
 		{
 			// remap the gamename if necessary
 			string GameName = Program.GameName;
@@ -47,6 +47,30 @@ namespace iPhonePackager
 				Directory.CreateDirectory(Config.ProvisionDirectory);
 			}
 
+			#region remove after we provide an install mechanism
+			// copy all of the provisions from the game directory to the library
+			foreach (string Provision in Directory.EnumerateFiles(Config.ProjectRootDirectory + "/Build/IOS/", "*.mobileprovision", SearchOption.AllDirectories))
+			{
+				if (!File.Exists(Config.ProvisionDirectory + Path.GetFileName(Provision)))
+				{
+					File.Copy(Provision, Config.ProvisionDirectory + Path.GetFileName(Provision), true);
+					FileInfo DestFileInfo = new FileInfo(Config.ProvisionDirectory + Path.GetFileName(Provision));
+					DestFileInfo.Attributes = DestFileInfo.Attributes & ~FileAttributes.ReadOnly;
+				}
+			}
+
+			// copy all of the provisions from the egine directory to the library
+			foreach (string Provision in Directory.EnumerateFiles(Config.EngineBuildDirectory, "*.mobileprovision", SearchOption.AllDirectories))
+			{
+				if (!File.Exists(Config.ProvisionDirectory + Path.GetFileName(Provision)))
+				{
+					File.Copy(Provision, Config.ProvisionDirectory + Path.GetFileName(Provision), true);
+					FileInfo DestFileInfo = new FileInfo(Config.ProvisionDirectory + Path.GetFileName(Provision));
+					DestFileInfo.Attributes = DestFileInfo.Attributes & ~FileAttributes.ReadOnly;
+				}
+			}
+			#endregion
+
 			// cache the provision library
 			Dictionary<string, MobileProvision> ProvisionLibrary = new Dictionary<string, MobileProvision>();
 			foreach (string Provision in Directory.EnumerateFiles(Config.ProvisionDirectory, "*.mobileprovision"))
@@ -58,7 +82,7 @@ namespace iPhonePackager
 			// check the cache for a provision matching the app id (com.company.Game)
 			foreach (KeyValuePair<string, MobileProvision> Pair in ProvisionLibrary)
 			{
-				if ((Pair.Value.ProvisionName == GameName || Pair.Value.ApplicationIdentifier.Contains(GameName)) && (!Config.bForDistribution || (Pair.Value.ProvisionedDeviceIDs.Count == 0 && !Pair.Value.bDebug)))
+				if (Pair.Value.ApplicationIdentifier.Contains(CFBundleIdentifier) && (!Config.bForDistribution || (Pair.Value.ProvisionedDeviceIDs.Count == 0 && !Pair.Value.bDebug)))
 				{
 					// check to see if we have a certificate for this provision
 					if (CodeSignatureBuilder.FindCertificate(Pair.Value) != null)
@@ -68,73 +92,6 @@ namespace iPhonePackager
 				}
 			}
 
-			#region remove after we provide an install mechanism
-			// check to see if we can find a provision in the game's build directory
-			string MobileProvisionFilename = FileOperations.FindPrefixedFile(Config.ProjectRootDirectory + "/Build/IOS/", GameName + ".mobileprovision");
-
-			if (!File.Exists(MobileProvisionFilename))
-			{
-				MobileProvisionFilename = FileOperations.FindPrefixedFile(Config.ProjectRootDirectory + "/Build/IOS/NotForLicensees/", GameName + ".mobileprovision");
-				if (!File.Exists(MobileProvisionFilename))
-				{
-					MobileProvisionFilename = FileOperations.FindAnyFileWithExtension(Config.ProjectRootDirectory + "/Build/IOS/", ".mobileprovision");
-				}
-			}
-			if (File.Exists(MobileProvisionFilename))
-			{
-				Directory.CreateDirectory(Config.ProvisionDirectory);
-				File.Copy(MobileProvisionFilename, Config.ProvisionDirectory + Path.GetFileName(MobileProvisionFilename), true);
-				FileInfo DestFileInfo = new FileInfo(Config.ProvisionDirectory + Path.GetFileName(MobileProvisionFilename));
-				DestFileInfo.Attributes = DestFileInfo.Attributes & ~FileAttributes.ReadOnly;
-				MobileProvisionFilename = Config.ProvisionDirectory + Path.GetFileName(MobileProvisionFilename);
-
-				// verify the provision is good
-				MobileProvision p = MobileProvisionParser.ParseFile(MobileProvisionFilename);
-				if ((p.ProvisionName == GameName || p.ApplicationIdentifier.Contains(GameName)) && (!Config.bForDistribution || (p.ProvisionedDeviceIDs.Count == 0 && !p.bDebug)))
-				{
-					// check to see if we have a certificate for this provision
-					if (CodeSignatureBuilder.FindCertificate(p) != null)
-					{
-						return MobileProvisionFilename;
-					}
-				}
-			}
-
-			// check for a distribution provision in the game's build directory
-			if (Config.bForDistribution)
-			{
-				MobileProvisionFilename = FileOperations.FindPrefixedFile(Config.ProjectRootDirectory + "/Build/IOS/", GameName + "_Distro.mobileprovision");
-
-				if (!File.Exists(MobileProvisionFilename))
-				{
-					MobileProvisionFilename = FileOperations.FindPrefixedFile(Config.ProjectRootDirectory + "/Build/IOS/NotForLicensees/", GameName + "_Distro.mobileprovision");
-					if (!File.Exists(MobileProvisionFilename))
-					{
-						MobileProvisionFilename = FileOperations.FindAnyFileWithExtension(Config.ProjectRootDirectory + "/Build/IOS/", "_Distro.mobileprovision");
-					}
-				}
-				if (File.Exists(MobileProvisionFilename))
-				{
-					Directory.CreateDirectory(Config.ProvisionDirectory);
-					File.Copy(MobileProvisionFilename, Config.ProvisionDirectory + Path.GetFileName(MobileProvisionFilename), true);
-					FileInfo DestFileInfo = new FileInfo(Config.ProvisionDirectory + Path.GetFileName(MobileProvisionFilename));
-					DestFileInfo.Attributes = DestFileInfo.Attributes & ~FileAttributes.ReadOnly;
-					MobileProvisionFilename = Config.ProvisionDirectory + Path.GetFileName(MobileProvisionFilename);
-
-					// verify the provision is good
-					MobileProvision p = MobileProvisionParser.ParseFile(MobileProvisionFilename);
-					if ((p.ProvisionName == GameName || p.ApplicationIdentifier.Contains(GameName)) && (!Config.bForDistribution || (p.ProvisionedDeviceIDs.Count == 0 && !p.bDebug)))
-					{
-						// check to see if we have a certificate for this provision
-						if (CodeSignatureBuilder.FindCertificate(p) != null)
-						{
-							return MobileProvisionFilename;
-						}
-					}
-				}
-			}
-			#endregion
-
 			// check the cache for a provision matching the wild card
 			foreach (KeyValuePair<string, MobileProvision> Pair in ProvisionLibrary)
 			{
@@ -148,75 +105,8 @@ namespace iPhonePackager
 				}
 			}
 
-			#region remove after we provide an install mechanism
-			// check to see if we can find a provision in the engine's build directory
-			MobileProvisionFilename = FileOperations.FindPrefixedFile(Config.EngineBuildDirectory, "UE4Game.mobileprovision");
-
-			if (!File.Exists(MobileProvisionFilename))
-			{
-				MobileProvisionFilename = FileOperations.FindPrefixedFile(Config.EngineBuildDirectory + "/NotForLicensees/", "UE4Game.mobileprovision");
-				if (!File.Exists(MobileProvisionFilename))
-				{
-					MobileProvisionFilename = FileOperations.FindAnyFileWithExtension(Config.EngineBuildDirectory, ".mobileprovision");
-				}
-			}
-			if (File.Exists(MobileProvisionFilename))
-			{
-				Directory.CreateDirectory(Config.ProvisionDirectory);
-				File.Copy(MobileProvisionFilename, Config.ProvisionDirectory + Path.GetFileName(MobileProvisionFilename), true);
-				FileInfo DestFileInfo = new FileInfo(Config.ProvisionDirectory + Path.GetFileName(MobileProvisionFilename));
-				DestFileInfo.Attributes = DestFileInfo.Attributes & ~FileAttributes.ReadOnly;
-				MobileProvisionFilename = Config.ProvisionDirectory + Path.GetFileName(MobileProvisionFilename);
-
-				// verify the provision is good
-				MobileProvision p = MobileProvisionParser.ParseFile(MobileProvisionFilename);
-				if ((p.ProvisionName.Contains("Wildcard") || p.ApplicationIdentifier.Contains("*")) && (!Config.bForDistribution || (p.ProvisionedDeviceIDs.Count == 0 && !p.bDebug)))
-				{
-					// check to see if we have a certificate for this provision
-					if (CodeSignatureBuilder.FindCertificate(p) != null)
-					{
-						return MobileProvisionFilename;
-					}
-				}
-			}
-
-			// check for a distribution provision in the game's build directory
-			if (Config.bForDistribution)
-			{
-				MobileProvisionFilename = FileOperations.FindPrefixedFile(Config.EngineBuildDirectory, "UE4Game_Distro.mobileprovision");
-
-				if (!File.Exists(MobileProvisionFilename))
-				{
-					MobileProvisionFilename = FileOperations.FindPrefixedFile(Config.EngineBuildDirectory + "/NotForLicensees/", "UE4Game_Distro.mobileprovision");
-					if (!File.Exists(MobileProvisionFilename))
-					{
-						MobileProvisionFilename = FileOperations.FindAnyFileWithExtension(Config.EngineBuildDirectory, ".mobileprovision");
-					}
-				}
-				if (File.Exists(MobileProvisionFilename))
-				{
-					Directory.CreateDirectory(Config.ProvisionDirectory);
-					File.Copy(MobileProvisionFilename, Config.ProvisionDirectory + Path.GetFileName(MobileProvisionFilename), true);
-					FileInfo DestFileInfo = new FileInfo(Config.ProvisionDirectory + Path.GetFileName(MobileProvisionFilename));
-					DestFileInfo.Attributes = DestFileInfo.Attributes & ~FileAttributes.ReadOnly;
-					MobileProvisionFilename = Config.ProvisionDirectory + Path.GetFileName(MobileProvisionFilename);
-
-					// verify the provision is good
-					MobileProvision p = MobileProvisionParser.ParseFile(MobileProvisionFilename);
-					if ((p.ProvisionName == GameName || p.ApplicationIdentifier.Contains(GameName)) && (!Config.bForDistribution || (p.ProvisionedDeviceIDs.Count == 0 && !p.bDebug)))
-					{
-						// check to see if we have a certificate for this provision
-						if (CodeSignatureBuilder.FindCertificate(p) != null)
-						{
-							return MobileProvisionFilename;
-						}
-					}
-				}
-			}
-			#endregion
-
 			// check to see if there is already an embedded provision
-			MobileProvisionFilename = Path.Combine(Config.RepackageStagingDirectory, "embedded.mobileprovision");
+			string MobileProvisionFilename = Path.Combine(Config.RepackageStagingDirectory, "embedded.mobileprovision");
 
 			return MobileProvisionFilename;
 		}
