fix for space in the path for the remote strip command
#ue4
#ios

diff --git a/Engine/Source/Programs/IOS/iPhonePackager/CompileTime.cs b/Engine/Source/Programs/IOS/iPhonePackager/CompileTime.cs
--- a/Engine/Source/Programs/IOS/iPhonePackager/CompileTime.cs
+++ b/Engine/Source/Programs/IOS/iPhonePackager/CompileTime.cs
@@ -324,7 +324,7 @@ namespace iPhonePackager
 
 			case "strip":
 				Program.Log( " ... stripping" );
-				DisplayCommandLine = XcodeDeveloperDir + "Platforms/iPhoneOS.platform/Developer/usr/bin/strip " + RemoteExecutablePath;
+				DisplayCommandLine = XcodeDeveloperDir + "Platforms/iPhoneOS.platform/Developer/usr/bin/strip '" + RemoteExecutablePath + "'";
 				CommandLine = "\"" + MacStagingRootDir + "\" " + DisplayCommandLine;
 				WorkingFolder = MacStagingRootDir;
 				break;
