updates to ULD, now called Unreal Mobile Tool
#ue4
#umt

diff --git a/Engine/Source/Programs/IOS/UnrealLaunchDaemon/Build/IOS/UnrealLaunchDaemon-Info.plist b/Engine/Source/Programs/IOS/UnrealLaunchDaemon/Build/IOS/UnrealLaunchDaemon-Info.plist
--- a/Engine/Source/Programs/IOS/UnrealLaunchDaemon/Build/IOS/UnrealLaunchDaemon-Info.plist
+++ b/Engine/Source/Programs/IOS/UnrealLaunchDaemon/Build/IOS/UnrealLaunchDaemon-Info.plist
@@ -16,11 +16,11 @@
 	<key>CFBundleDevelopmentRegion</key>
 	<string>English</string>
 	<key>CFBundleDisplayName</key>
-	<string>UnrealLauncher</string>
+	<string>Mobile Tool</string>
 	<key>CFBundleExecutable</key>
 	<string>${EXECUTABLE_NAME}</string>
 	<key>CFBundleIdentifier</key>
-	<string>com.epicgames.UnrealLauncher</string>
+	<string>com.epicgames.UnrealMobileTool</string>
 	<key>CFBundleInfoDictionaryVersion</key>
 	<string>6.0</string>
 	<key>CFBundleName</key>
diff --git a/Engine/Source/Programs/IOS/UnrealLaunchDaemon/Private/IOS/IOSLaunchDaemonAppMain.cpp b/Engine/Source/Programs/IOS/UnrealLaunchDaemon/Private/IOS/IOSLaunchDaemonAppMain.cpp
--- a/Engine/Source/Programs/IOS/UnrealLaunchDaemon/Private/IOS/IOSLaunchDaemonAppMain.cpp
+++ b/Engine/Source/Programs/IOS/UnrealLaunchDaemon/Private/IOS/IOSLaunchDaemonAppMain.cpp
@@ -18,6 +18,18 @@
 FString GSavedCommandLine;
 FLaunchDaemonMessageHandler GCommandSystem;
 
+void FAppEntry::Suspend()
+{
+}
+
+void FAppEntry::Resume()
+{
+}
+
+void FAppEntry::SuspendTick()
+{
+}
+
 void FAppEntry::PreInit(IOSAppDelegate* AppDelegate, UIApplication* Application)
 {
 	Application.idleTimerDisabled = YES;
@@ -124,7 +136,7 @@ int main(int argc, char *argv[])
 		GSavedCommandLine += ANSI_TO_TCHAR(argv[Option]);
 	}
 
-	FIOSCommandLineHelper::InitCommandArgs(FString());
+	FIOSCommandLineHelper::InitCommandArgs(TEXT("-messaging"));
 
     @autoreleasepool {
         return UIApplicationMain(argc, argv, nil, NSStringFromClass([IOSAppDelegate class]));
diff --git a/Engine/Source/Programs/IOS/UnrealLaunchDaemon/UnrealLaunchDaemon.Build.cs b/Engine/Source/Programs/IOS/UnrealLaunchDaemon/UnrealLaunchDaemon.Build.cs
--- a/Engine/Source/Programs/IOS/UnrealLaunchDaemon/UnrealLaunchDaemon.Build.cs
+++ b/Engine/Source/Programs/IOS/UnrealLaunchDaemon/UnrealLaunchDaemon.Build.cs
@@ -24,6 +24,7 @@ public class UnrealLaunchDaemon : ModuleRules
 				"StandaloneRenderer",
 				"LaunchDaemonMessages",
 				"Projects",
+				"Messaging",
 				"UdpMessaging"
 			}
 		);
