Fixed iPhonePackager compile warning (which is an error on Mono)

diff --git a/Engine/Source/Programs/IOS/iPhonePackager/CodeSigning.cs b/Engine/Source/Programs/IOS/iPhonePackager/CodeSigning.cs
--- a/Engine/Source/Programs/IOS/iPhonePackager/CodeSigning.cs
+++ b/Engine/Source/Programs/IOS/iPhonePackager/CodeSigning.cs
@@ -509,7 +509,6 @@ namespace iPhonePackager
 				// Adjust the header and load command to have the correct size for the code sign blob
 				WritingContext OutputExeContext = new WritingContext(new BinaryWriter(OutputExeStream));
 
-				long StartPosition = (CodeSigningBlobLC.BlobFileOffset - (long)LinkEditSegmentLC.FileOffset);
 				long BlobLength = DummyPayload.Length;
 
 				long NonCodeSigSize = (long)LinkEditSegmentLC.FileSize - CodeSigningBlobLC.BlobFileSize;
