- Added openssl.exe to UE4 [ttp 340264]

diff --git a/Engine/Source/Programs/IOS/iPhonePackager/GenerateSigningRequestDialog.cs b/Engine/Source/Programs/IOS/iPhonePackager/GenerateSigningRequestDialog.cs
--- a/Engine/Source/Programs/IOS/iPhonePackager/GenerateSigningRequestDialog.cs
+++ b/Engine/Source/Programs/IOS/iPhonePackager/GenerateSigningRequestDialog.cs
@@ -49,7 +49,7 @@ namespace iPhonePackager
 			string OpenSSLPath = Path.GetDirectoryName( Application.ExecutablePath ) + @"\openssl.exe";
 			if (!File.Exists( OpenSSLPath ))
 			{
-				MessageBox.Show("A version of OpenSSL is required to generate certificate requests.  Please place OpenSSL.exe in Binaries\\IPhone", Config.AppDisplayName, MessageBoxButtons.OK, MessageBoxIcon.Error);
+				MessageBox.Show("A version of OpenSSL is required to generate certificate requests.  Please place OpenSSL.exe in Binaries\\DotNET\\IOS", Config.AppDisplayName, MessageBoxButtons.OK, MessageBoxIcon.Error);
 				return;
 			}
 
