ipp - Fix for folders not created

diff --git a/Engine/Source/Programs/IOS/iPhonePackager/MobileProvisionUtilities.cs b/Engine/Source/Programs/IOS/iPhonePackager/MobileProvisionUtilities.cs
--- a/Engine/Source/Programs/IOS/iPhonePackager/MobileProvisionUtilities.cs
+++ b/Engine/Source/Programs/IOS/iPhonePackager/MobileProvisionUtilities.cs
@@ -34,7 +34,7 @@ namespace iPhonePackager
 			string GameName = Program.GameName;
 			if (GameName == "UE4Game")
 			{
-				if(Config.ProjectFile.Length > 0)
+				if (Config.ProjectFile.Length > 0)
 				{
 					GameName = Path.GetFileNameWithoutExtension(Config.ProjectFile);
 				}
@@ -48,24 +48,35 @@ namespace iPhonePackager
 
 			#region remove after we provide an install mechanism
 			// copy all of the provisions from the game directory to the library
-			foreach (string Provision in Directory.EnumerateFiles(Path.GetDirectoryName(Config.ProjectFile) + "/Build/IOS/", "*.mobileprovision", SearchOption.AllDirectories))
 			{
-				if (!File.Exists(Config.ProvisionDirectory + Path.GetFileName(Provision)))
+				var ProjectFileBuildIOSPath = Path.GetDirectoryName(Config.ProjectFile) + "/Build/IOS/";
+				if (Directory.Exists(ProjectFileBuildIOSPath))
 				{
-					File.Copy(Provision, Config.ProvisionDirectory + Path.GetFileName(Provision), true);
-					FileInfo DestFileInfo = new FileInfo(Config.ProvisionDirectory + Path.GetFileName(Provision));
-					DestFileInfo.Attributes = DestFileInfo.Attributes & ~FileAttributes.ReadOnly;
+					foreach (string Provision in Directory.EnumerateFiles(ProjectFileBuildIOSPath, "*.mobileprovision", SearchOption.AllDirectories))
+					{
+						if (!File.Exists(Config.ProvisionDirectory + Path.GetFileName(Provision)))
+						{
+							File.Copy(Provision, Config.ProvisionDirectory + Path.GetFileName(Provision), true);
+							FileInfo DestFileInfo = new FileInfo(Config.ProvisionDirectory + Path.GetFileName(Provision));
+							DestFileInfo.Attributes = DestFileInfo.Attributes & ~FileAttributes.ReadOnly;
+						}
+					}
 				}
 			}
 
-			// copy all of the provisions from the egine directory to the library
-			foreach (string Provision in Directory.EnumerateFiles(Config.EngineBuildDirectory, "*.mobileprovision", SearchOption.AllDirectories))
+			// copy all of the provisions from the engine directory to the library
 			{
-				if (!File.Exists(Config.ProvisionDirectory + Path.GetFileName(Provision)))
+				if (Directory.Exists(Config.EngineBuildDirectory))
 				{
-					File.Copy(Provision, Config.ProvisionDirectory + Path.GetFileName(Provision), true);
-					FileInfo DestFileInfo = new FileInfo(Config.ProvisionDirectory + Path.GetFileName(Provision));
-					DestFileInfo.Attributes = DestFileInfo.Attributes & ~FileAttributes.ReadOnly;
+					foreach (string Provision in Directory.EnumerateFiles(Config.EngineBuildDirectory, "*.mobileprovision", SearchOption.AllDirectories))
+					{
+						if (!File.Exists(Config.ProvisionDirectory + Path.GetFileName(Provision)))
+						{
+							File.Copy(Provision, Config.ProvisionDirectory + Path.GetFileName(Provision), true);
+							FileInfo DestFileInfo = new FileInfo(Config.ProvisionDirectory + Path.GetFileName(Provision));
+							DestFileInfo.Attributes = DestFileInfo.Attributes & ~FileAttributes.ReadOnly;
+						}
+					}
 				}
 			}
 			#endregion
