fix for not updating mobile provisions when the provision is newer than the one in the cache
TTP341334
#ue4
#ios

diff --git a/Engine/Source/Programs/IOS/iPhonePackager/MobileProvisionUtilities.cs b/Engine/Source/Programs/IOS/iPhonePackager/MobileProvisionUtilities.cs
--- a/Engine/Source/Programs/IOS/iPhonePackager/MobileProvisionUtilities.cs
+++ b/Engine/Source/Programs/IOS/iPhonePackager/MobileProvisionUtilities.cs
@@ -28,7 +28,7 @@ namespace iPhonePackager
 		public bool bDebug;
 		public Utilities.PListHelper Data;
 
-		public static string FindCompatibleProvision(string CFBundleIdentifier)
+		public static string FindCompatibleProvision(string CFBundleIdentifier, bool bCheckCert = true)
 		{
 			// remap the gamename if necessary
 			string GameName = Program.GameName;
@@ -54,7 +54,7 @@ namespace iPhonePackager
 				{
 					foreach (string Provision in Directory.EnumerateFiles(ProjectFileBuildIOSPath, "*.mobileprovision", SearchOption.AllDirectories))
 					{
-						if (!File.Exists(Config.ProvisionDirectory + Path.GetFileName(Provision)))
+						if (!File.Exists(Config.ProvisionDirectory + Path.GetFileName(Provision)) || File.GetLastWriteTime(Config.ProvisionDirectory + Path.GetFileName(Provision)) < File.GetLastWriteTime(Provision))
 						{
 							File.Copy(Provision, Config.ProvisionDirectory + Path.GetFileName(Provision), true);
 							FileInfo DestFileInfo = new FileInfo(Config.ProvisionDirectory + Path.GetFileName(Provision));
@@ -70,7 +70,7 @@ namespace iPhonePackager
 				{
 					foreach (string Provision in Directory.EnumerateFiles(Config.EngineBuildDirectory, "*.mobileprovision", SearchOption.AllDirectories))
 					{
-						if (!File.Exists(Config.ProvisionDirectory + Path.GetFileName(Provision)))
+						if (!File.Exists(Config.ProvisionDirectory + Path.GetFileName(Provision)) || File.GetLastWriteTime(Config.ProvisionDirectory + Path.GetFileName(Provision)) < File.GetLastWriteTime(Provision))
 						{
 							File.Copy(Provision, Config.ProvisionDirectory + Path.GetFileName(Provision), true);
 							FileInfo DestFileInfo = new FileInfo(Config.ProvisionDirectory + Path.GetFileName(Provision));
@@ -95,7 +95,7 @@ namespace iPhonePackager
 				if (Pair.Value.ApplicationIdentifier.Contains(CFBundleIdentifier) && (!Config.bForDistribution || (Pair.Value.ProvisionedDeviceIDs.Count == 0 && !Pair.Value.bDebug)))
 				{
 					// check to see if we have a certificate for this provision
-					if (CodeSignatureBuilder.FindCertificate(Pair.Value) != null)
+					if (!bCheckCert || CodeSignatureBuilder.FindCertificate(Pair.Value) != null)
 					{
 						return Pair.Key;
 					}
@@ -108,7 +108,7 @@ namespace iPhonePackager
 				if ((Pair.Value.ProvisionName.Contains("Wildcard") || Pair.Value.ApplicationIdentifier.Contains("*")) && (!Config.bForDistribution || (Pair.Value.ProvisionedDeviceIDs.Count == 0 && !Pair.Value.bDebug)))
 				{
 					// check to see if we have a certificate for this provision
-					if (CodeSignatureBuilder.FindCertificate(Pair.Value) != null)
+					if (!bCheckCert || CodeSignatureBuilder.FindCertificate(Pair.Value) != null)
 					{
 						return Pair.Key;
 					}
