fix for IPP GUI causing a crash
#ue4
#ios

diff --git a/Engine/Source/Programs/IOS/iPhonePackager/CodeSigning.cs b/Engine/Source/Programs/IOS/iPhonePackager/CodeSigning.cs
--- a/Engine/Source/Programs/IOS/iPhonePackager/CodeSigning.cs
+++ b/Engine/Source/Programs/IOS/iPhonePackager/CodeSigning.cs
@@ -77,14 +77,18 @@ namespace iPhonePackager
 			bHasOverridesFile = File.Exists(Config.GetPlistOverrideFilename());
 
 			// Load Info.plist, which guides nearly everything else
-			string plistFile = Path.GetDirectoryName(Config.ProjectFile) + "/Build/IOS/Info.plist";
-			if (!File.Exists(plistFile))
+			string plistFile = Config.EngineBuildDirectory + "/UE4Game-Info.plist";
+			if (!string.IsNullOrEmpty(Config.ProjectFile))
 			{
-				plistFile = Path.GetDirectoryName(Config.ProjectFile) + "/Build/IOS/" + Path.GetFileNameWithoutExtension(Config.ProjectFile) + "-Info.plist";
-
+				plistFile = Path.GetDirectoryName(Config.ProjectFile) + "/Build/IOS/Info.plist";
 				if (!File.Exists(plistFile))
 				{
-					plistFile = Config.EngineBuildDirectory + "/UE4Game-Info.plist";
+					plistFile = Path.GetDirectoryName(Config.ProjectFile) + "/Build/IOS/" + Path.GetFileNameWithoutExtension(Config.ProjectFile) + "-Info.plist";
+
+					if (!File.Exists(plistFile))
+					{
+						plistFile = Config.EngineBuildDirectory + "/UE4Game-Info.plist";
+					}
 				}
 			}
 			Utilities.PListHelper Info = null;
