- Updated IPP to work with new provision searching. Now requires a uproject for the GUI [ttp 342924]
- Removed old cruft from the GUI

diff --git a/Engine/Source/Programs/IOS/iPhonePackager/ToolsHub.cs b/Engine/Source/Programs/IOS/iPhonePackager/ToolsHub.cs
--- a/Engine/Source/Programs/IOS/iPhonePackager/ToolsHub.cs
+++ b/Engine/Source/Programs/IOS/iPhonePackager/ToolsHub.cs
@@ -87,7 +87,7 @@ namespace iPhonePackager
 		public static ToolsHub CreateShowingTools()
 		{
 			ToolsHub Result = new ToolsHub();
-			Result.tabControl1.SelectTab(Result.tabPage3);
+			Result.tabControl1.SelectTab(Result.tabPage2);
 			return Result;
 		}
 
@@ -106,9 +106,9 @@ namespace iPhonePackager
 
 			MobileProvisionCheck2.Image = MobileProvisionCheck.Image = CheckStateImages[Provision != null];
 			CertificatePresentCheck2.Image = CertificatePresentCheck.Image = CheckStateImages[Cert != null];
-			OverridesPresentCheck2.Image = OverridesPresentCheck.Image = CheckStateImages[bOverridesExists];
+//			OverridesPresentCheck2.Image = OverridesPresentCheck.Image = CheckStateImages[bOverridesExists];
 
-			ReadyToPackageButton.Enabled = bOverridesExists && (Provision != null) && (Cert != null);
+//			ReadyToPackageButton.Enabled = /*bOverridesExists && */(Provision != null) && (Cert != null);
 		}
 
 		private void CreateCSRButton_Click(object sender, EventArgs e)
@@ -139,9 +139,12 @@ namespace iPhonePackager
 					MobileProvision Provision = MobileProvisionParser.ParseFile(ProvisionFilename);
 					bIsDistribution = IsProfileForDistribution(Provision);
 
+					// use the input filename if the GameName is empty
+					string DestName = string.IsNullOrEmpty(Program.GameName) ? Path.GetFileNameWithoutExtension(ProvisionFilename) : Program.GameName;
+
 					// Copy the file into the destination location
 					string EffectivePrefix = bIsDistribution ? "Distro_" : Config.SigningPrefix;
-					string DestinationFilename = Path.Combine(Config.ProvisionDirectory, EffectivePrefix + Program.GameName + ".mobileprovision");
+					string DestinationFilename = Path.Combine(Config.ProvisionDirectory, EffectivePrefix + DestName + ".mobileprovision");
 
 					if (File.Exists(DestinationFilename))
 					{
diff --git a/Engine/Source/Programs/IOS/iPhonePackager/iPhonePackager.cs b/Engine/Source/Programs/IOS/iPhonePackager/iPhonePackager.cs
--- a/Engine/Source/Programs/IOS/iPhonePackager/iPhonePackager.cs
+++ b/Engine/Source/Programs/IOS/iPhonePackager/iPhonePackager.cs
@@ -86,13 +86,34 @@ namespace iPhonePackager
 		{
 			if (Arguments.Length == 0)
 			{
-				Program.Warning("No arguments specified, assuming GUI mode");
-				Arguments = new string[] { "gui" };
+				StartVisuals();
+
+				// we NEED a project, so show a uproject picker
+				string UProjectFile;
+				string StartingDir = "";
+				if (ToolsHub.ShowOpenFileDialog("Unreal Project Files (*.uproject)|*.uproject;", "IPhonePackager now requires a .uproject file to continue", "mobileprovision", "", ref StartingDir, out UProjectFile))
+				{
+					Arguments = new string[] { UProjectFile };
+				}
+				else
+				{
+					return false;
+				}
 			}
 
 			if (Arguments.Length == 1)
 			{
-				MainCommand = Arguments[0];
+				// if the only argument is a uproject, then assume gui mode, with the uproject as the project
+				if (Arguments[0].EndsWith(".uproject"))
+				{
+					Config.ProjectFile = GamePath = Arguments[0];
+
+					MainCommand = "gui";
+				}
+				else
+				{
+					MainCommand = Arguments[0];
+				}
 			}
 			else if (Arguments.Length == 2)
 			{
@@ -412,8 +433,15 @@ namespace iPhonePackager
 
 		delegate Form CreateFormDelegate();
 
+		static bool bVisualsStarted = false;
 		static void StartVisuals()
 		{
+			if (bVisualsStarted)
+			{
+				return;
+			}
+			bVisualsStarted = true;
+
 			Application.EnableVisualStyles();
 			Application.SetCompatibleTextRenderingDefault(false);
 
