- Put a Install IPA button back in.

diff --git a/Engine/Source/Programs/IOS/iPhonePackager/ToolsHub.cs b/Engine/Source/Programs/IOS/iPhonePackager/ToolsHub.cs
--- a/Engine/Source/Programs/IOS/iPhonePackager/ToolsHub.cs
+++ b/Engine/Source/Programs/IOS/iPhonePackager/ToolsHub.cs
@@ -95,6 +95,13 @@ namespace iPhonePackager
 		{
 			// Check the status of the various steps
 			UpdateStatus();
+
+			CreateCSRButton.Enabled =
+			ImportMobileProvisionButton.Enabled =
+			ImportProvisionButton2.Enabled =
+			ImportCertificateButton.Enabled =
+			ImportCertificateButton2.Enabled =
+				!string.IsNullOrEmpty(Config.ProjectFile);
 		}
 
 		void UpdateStatus()
@@ -339,5 +346,10 @@ namespace iPhonePackager
 			Process.Start(PSI);
 		}
 
+		private void InstallIPA_Click(object sender, EventArgs e)
+		{
+			InstallIPAButton_Click(sender, e);
+		}
+
 	}
 }
\ No newline at end of file
diff --git a/Engine/Source/Programs/IOS/iPhonePackager/iPhonePackager.cs b/Engine/Source/Programs/IOS/iPhonePackager/iPhonePackager.cs
--- a/Engine/Source/Programs/IOS/iPhonePackager/iPhonePackager.cs
+++ b/Engine/Source/Programs/IOS/iPhonePackager/iPhonePackager.cs
@@ -91,13 +91,13 @@ namespace iPhonePackager
 				// we NEED a project, so show a uproject picker
 				string UProjectFile;
 				string StartingDir = "";
-				if (ToolsHub.ShowOpenFileDialog("Unreal Project Files (*.uproject)|*.uproject;", "IPhonePackager now requires a .uproject file to continue", "mobileprovision", "", ref StartingDir, out UProjectFile))
+				if (ToolsHub.ShowOpenFileDialog("Unreal Project Files (*.uproject)|*.uproject;", "IPhonePackager now requires a .uproject file for certificate/provisioning setup", "mobileprovision", "", ref StartingDir, out UProjectFile))
 				{
 					Arguments = new string[] { UProjectFile };
 				}
 				else
 				{
-					return false;
+					Arguments = new string[] { "gui" };
 				}
 			}
 
