fix for some error codes from IPP
TTP343083
#ue4
#ipp

diff --git a/Engine/Source/Programs/IOS/iPhonePackager/CodeSigning.cs b/Engine/Source/Programs/IOS/iPhonePackager/CodeSigning.cs
--- a/Engine/Source/Programs/IOS/iPhonePackager/CodeSigning.cs
+++ b/Engine/Source/Programs/IOS/iPhonePackager/CodeSigning.cs
@@ -175,6 +175,8 @@ namespace iPhonePackager
 			{
 				Program.Error("Could not find {0}.mobileprovision in {1} (error: '{2}').  Please follow the setup instructions to get a mobile provision from the Apple Developer site.",
 					Program.GameName, Path.GetFullPath(Config.BuildDirectory), ex.Message);
+				Program.ReturnCode = (int)ErrorCodes.Error_ProvisionNotFound;
+				throw ex;
 			}
 
 			return Result;
@@ -374,6 +376,8 @@ namespace iPhonePackager
 			{
 				// Failed to find a cert already installed or to install, cannot proceed any futher
 				Program.Error("... Failed to find a certificate that matches the mobile provision to be used for code signing");
+				Program.ReturnCode = (int)ErrorCodes.Error_CertificateNotFound;
+				throw new InvalidDataException("Certificate not found!");
 			}
 			else
 			{
diff --git a/Engine/Source/Programs/IOS/iPhonePackager/iPhonePackager.cs b/Engine/Source/Programs/IOS/iPhonePackager/iPhonePackager.cs
--- a/Engine/Source/Programs/IOS/iPhonePackager/iPhonePackager.cs
+++ b/Engine/Source/Programs/IOS/iPhonePackager/iPhonePackager.cs
@@ -603,7 +603,7 @@ namespace iPhonePackager
 						string dllPath = Microsoft.Win32.Registry.GetValue("HKEY_LOCAL_MACHINE\\SOFTWARE\\Wow6432Node\\Apple Inc.\\Apple Mobile Device Support\\Shared", "iTunesMobileDeviceDLL", null) as string; 
 						if (String.IsNullOrEmpty(dllPath) || !File.Exists(dllPath))
 						{
-							Error("iTunes Not Found!!", ErrorCodes.Error_SDKNotFound);
+							Error("iTunes Not Found!!", (int)ErrorCodes.Error_SDKNotFound);
 						}
 						else
 						{
@@ -614,19 +614,19 @@ namespace iPhonePackager
 							bool foundPlist = CodeSignatureBuilder.FindRequiredFiles(out Provision, out Cert, out bHasOverrides);
 							if (!foundPlist)
 							{
-								Error("Could not find a valid plist file!!", ErrorCodes.Error_InfoPListNotFound);
+								Error("Could not find a valid plist file!!", (int)ErrorCodes.Error_InfoPListNotFound);
 							}
 							else if (Provision == null && Cert == null)
 							{
-								Error("No Provision or cert found!!", ErrorCodes.Error_ProvisionAndCertificateNotFound);
+								Error("No Provision or cert found!!", (int)ErrorCodes.Error_ProvisionAndCertificateNotFound);
 							}
 							else if (Provision == null)
 							{
-								Error("No Provision found!!", ErrorCodes.Error_ProvisionNotFound);
+								Error("No Provision found!!", (int)ErrorCodes.Error_ProvisionNotFound);
 							}
 							else if (Cert == null)
 							{
-								Error("No Signing Certificate found!!", ErrorCodes.Error_CertificateNotFound);
+								Error("No Signing Certificate found!!", (int)ErrorCodes.Error_CertificateNotFound);
 							}
 						}
 						break;
