- Put some stuff back in to IPP

diff --git a/Engine/Source/Programs/IOS/iPhonePackager/GraphicalResignTool.cs b/Engine/Source/Programs/IOS/iPhonePackager/GraphicalResignTool.cs
--- a/Engine/Source/Programs/IOS/iPhonePackager/GraphicalResignTool.cs
+++ b/Engine/Source/Programs/IOS/iPhonePackager/GraphicalResignTool.cs
@@ -18,6 +18,7 @@ namespace iPhonePackager
 {
 	public partial class GraphicalResignTool : Form
 	{
+
 		public static GraphicalResignTool ActiveInstance = null;
 
 		private string ExportImportPListFilename = "";
@@ -51,25 +52,6 @@ namespace iPhonePackager
 		/// </summary>
 		void InitializeGameSelectionBox()
 		{
-			string InstallPath = Path.GetDirectoryName(Application.ExecutablePath) + Path.DirectorySeparatorChar + ".." + Path.DirectorySeparatorChar + ".." + Path.DirectorySeparatorChar;
-			string[] GameDirectories = Directory.GetDirectories(InstallPath, "*Game");
-
-			// Fetch the currently active game name
-			string ActiveGameName = (Program.GameName != null) ? Program.GameName : "UDKGame";
-
-			// Add the game names, highlighting the active one
-			GameSelectionBox.Items.Clear();
-			foreach (string GameDir in GameDirectories)
-			{
-				int CutIndex = GameDir.LastIndexOf(Path.DirectorySeparatorChar);
-				string GameName = GameDir.Substring(CutIndex + 1);
-
-				int GameIndexInComboBox = GameSelectionBox.Items.Add(GameName);
-				if (GameName.ToUpperInvariant() == ActiveGameName.ToUpperInvariant())
-				{
-					GameSelectionBox.SelectedIndex = GameIndexInComboBox;
-				}
-			}
 		}
 
 		private void GraphicalResignTool_Load(object sender, EventArgs e)
@@ -433,90 +415,6 @@ namespace iPhonePackager
 
 		void RefreshProvisionsList()
 		{
-			// Scan for all mobile provisions installed for the current game
-			Dictionary<string, MobileProvision> Provisions = new Dictionary<string, MobileProvision>();
-
-			Program.ProgressDialog.OnBeginBackgroundWork = delegate
-			{
-				Program.Log("Scanning for mobile provisions in '{0}'...", Config.BuildDirectory);
-
-				try
-				{
-					// Find all the provisions to check
-					string[] ProvisionFilenames = Directory.GetFiles(Config.BuildDirectory, "*.mobileprovision");
-
-					// Open each one in turn
-					int ProgressIndex = 0;
-					foreach (string Filename in ProvisionFilenames)
-					{
-						++ProgressIndex;
-						MobileProvision Provision = null;
-						try
-						{
-							Program.ProgressIndex = (ProgressIndex * 100) / ProvisionFilenames.Length;
-							Program.Log("... Loading provision '{0}'", Filename);
-
-							Provision = MobileProvisionParser.ParseFile(Filename);
-							Provision.Tag = CodeSignatureBuilder.FindCertificate(Provision);
-						}
-						catch
-						{
-						}
-						Provisions.Add(Filename, Provision);
-					}
-				}
-				catch
-				{
-				}
-			};
-
-			if (Program.ProgressDialog.ShowDialog() == DialogResult.OK)
-			{
-				KnownProvisionsTree.BeginUpdate();
-				KnownProvisionsTree.Nodes.Clear();
-
-				// Visualize the found provisions
-				foreach (var KVP in Provisions)
-				{
-					string Filename = KVP.Key;
-					MobileProvision Provision = KVP.Value;
-
-					TreeNode NewItem = KnownProvisionsTree.Nodes.Add(Path.GetFileNameWithoutExtension(Filename));
-					NewItem.Tag = KVP;
-
-					if (Provision != null)
-					{
-						X509Certificate2 Cert = Provision.Tag as X509Certificate2;
-
-						NewItem.Nodes.Add(String.Format("Name: '{0}'", Provision.ProvisionName));
-						NewItem.Nodes.Add(String.Format("Type: {0}", ToolsHub.IsProfileForDistribution(Provision) ? "Distribution" : "Development"));
-
-						string Prefix = Path.GetFileNameWithoutExtension(Filename).Replace(Program.GameName, "");
-						if (Prefix != "")
-						{
-							NewItem.Nodes.Add(String.Format("Signing prefix: '{0}'", Prefix));
-						}
-
-						NewItem.Nodes.Add(String.Format("Matching Cert: '{0}'", (Cert != null) ? Cert.FriendlyName : "None (missing)"));
-
-						// Add the device list
-						TreeNode DeviceList = NewItem.Nodes.Add(String.Format("Allowed devices: {0}", Provision.ProvisionedDeviceIDs.Count));
-						foreach (string DeviceID in Provision.ProvisionedDeviceIDs)
-						{
-							DeviceList.Nodes.Add(DeviceID);
-						}
-						DeviceList.Collapse();
-					}
-					else
-					{
-						NewItem.Nodes.Add("Error: Failed to parse");
-					}
-
-					NewItem.Expand();
-				}
-
-				KnownProvisionsTree.EndUpdate();
-			}
 		}
 
 		private void RefreshProvisionListButton_Click(object sender, EventArgs e)
@@ -534,11 +432,7 @@ namespace iPhonePackager
 
 		private void TabBook_SelectedIndexChanged(object sender, EventArgs e)
 		{
-			if (TabBook.SelectedTab == ProvisionCertPage)
-			{
-				RefreshProvisionsList();
-			}
-			else if (TabBook.SelectedTab == DeploymentPage)
+			if (TabBook.SelectedTab == DeploymentPage)
 			{
 				ScanConnectedDevices();
 			}
@@ -561,29 +455,10 @@ namespace iPhonePackager
 
 		private void copyToolStripMenuItem_Click(object sender, EventArgs e)
 		{
-			TreeNode Node = KnownProvisionsTree.SelectedNode;
-			if (Node != null)
-			{
-				Clipboard.SetText(Node.Text);
-			}
 		}
 
 		private void copyFullPathToProvisionToolStripMenuItem_Click(object sender, EventArgs e)
 		{
-			//
-			TreeNode Node = KnownProvisionsTree.SelectedNode;
-			if (Node != null)
-			{
-				// Walk up to the root
-				while (Node.Parent != null)
-				{
-					Node = Node.Parent;
-				}
-
-				// The tag of the root is a pair of filename, mobile provision
-				KeyValuePair<string, MobileProvision> KVP = (KeyValuePair<string, MobileProvision>)Node.Tag;
-				Clipboard.SetText(KVP.Key);
-			}
 		}
 
 		private void GraphicalResignTool_Shown(object sender, EventArgs e)
@@ -768,50 +643,6 @@ namespace iPhonePackager
 
 		private void exportCertificateToolStripMenuItem_Click(object sender, EventArgs e)
 		{
-			TreeNode Node = KnownProvisionsTree.SelectedNode;
-			if (Node != null)
-			{
-				// Walk up to the root
-				while (Node.Parent != null)
-				{
-					Node = Node.Parent;
-				}
-
-				// The tag of the root is a pair of filename, mobile provision
-				KeyValuePair<string, MobileProvision> KVP = (KeyValuePair<string, MobileProvision>)Node.Tag;
-
-				MobileProvision Provision = KVP.Value;
-				X509Certificate2 Cert = Provision.Tag as X509Certificate2;
-
-				if (Cert != null)
-				{
-					// Ask for an export filename
-					saveFileDialog1.DefaultExt = "p12";
-					saveFileDialog1.Title = "Choose a destination for the exported PKCS#12 certificate + key chain";
-					saveFileDialog1.Filter = ToolsHub.JustP12Certificates;
-					saveFileDialog1.FileName = "";
-
-					string CWD = Directory.GetCurrentDirectory();
-					bool bDialogSucceeded = (saveFileDialog1.ShowDialog() == DialogResult.OK);
-					Directory.SetCurrentDirectory(CWD);
-
-					if (bDialogSucceeded)
-					{
-						string TargetFilename = saveFileDialog1.FileName;
- 
-						// Ask for an export password
-						string ExportPassword;
-						if (PasswordDialog.RequestPassword(out ExportPassword))
-						{
-							File.WriteAllBytes(TargetFilename, Cert.Export(X509ContentType.Pkcs12, ExportPassword));
-						}
-					}
-				}
-				else
-				{
-					ShowErrorMessage("The selected mobile provision doesn't have an associated certificate to export!");
-				}
-			}
 		}
 	}
 }
diff --git a/Engine/Source/Programs/IOS/iPhonePackager/ToolsHub.cs b/Engine/Source/Programs/IOS/iPhonePackager/ToolsHub.cs
--- a/Engine/Source/Programs/IOS/iPhonePackager/ToolsHub.cs
+++ b/Engine/Source/Programs/IOS/iPhonePackager/ToolsHub.cs
@@ -87,7 +87,7 @@ namespace iPhonePackager
 		public static ToolsHub CreateShowingTools()
 		{
 			ToolsHub Result = new ToolsHub();
-			Result.tabControl1.SelectTab(Result.tabPage2);
+			Result.tabControl1.SelectTab(Result.tabPage1);
 			return Result;
 		}
 
@@ -307,9 +307,9 @@ namespace iPhonePackager
 
 		private void ProvisionCertToolsButton_Click(object sender, EventArgs e)
 		{
-			GraphicalResignTool ResignTool = GraphicalResignTool.GetActiveInstance();
-			ResignTool.TabBook.SelectTab(ResignTool.ProvisionCertPage);
-			ResignTool.Show();
+// 			GraphicalResignTool ResignTool = GraphicalResignTool.GetActiveInstance();
+// 			ResignTool.TabBook.SelectTab(ResignTool.ProvisionCertPage);
+// 			ResignTool.Show();
 		}
 
 		private void OtherDeployToolsButton_Click(object sender, EventArgs e)
