fix for IPP holding on to IPA file handle after deploying to device
TTP324571
#ue4

diff --git a/Engine/Source/Programs/IOS/iPhonePackager/DeployTime.cs b/Engine/Source/Programs/IOS/iPhonePackager/DeployTime.cs
--- a/Engine/Source/Programs/IOS/iPhonePackager/DeployTime.cs
+++ b/Engine/Source/Programs/IOS/iPhonePackager/DeployTime.cs
@@ -69,6 +69,7 @@ namespace iPhonePackager
 				Program.Warning(String.Format("Couldn't find an embedded mobile provision ({0})", ex.Message));
 				Provision = null;
 			}
+			Zip.Close();
 
 			if (Provision != null)
 			{
diff --git a/Engine/Source/Programs/IOS/iPhonePackager/FileOperations.cs b/Engine/Source/Programs/IOS/iPhonePackager/FileOperations.cs
--- a/Engine/Source/Programs/IOS/iPhonePackager/FileOperations.cs
+++ b/Engine/Source/Programs/IOS/iPhonePackager/FileOperations.cs
@@ -867,6 +867,7 @@ namespace iPhonePackager
 
 			public override void Close()
 			{
+				Zip.Dispose();
 				Zip = null;
 			}
 
