Platform: iPhonePackager now filters out expired mobile provisions and signing certificates
- Add checking of validity dates on mobile provisions and signing certificates, choosing only ones that are currently valid
- Add extended logging for found mobile provisions and signing certificates, to help debug possible provisioning issues in the future

- Remove duplicated Apple chain certificates from DotNET/IOS and switch to using the ones in Engine/Build/IOS
- Various other output logging improvements and time stamps during code signing
- Remove rotted -interactive option and associated code
- Fix unused variable warning on mono

For release notes: iOS deployment from PC requires the machine date/time to be reasonably accurate (date/time checks are now in place to filter out expired mobile provisions and signing certificates, preventing various deployment issues, but they rely on the machine time being correct)

diff --git a/Engine/Source/Programs/IOS/iPhonePackager/CodeSigning.cs b/Engine/Source/Programs/IOS/iPhonePackager/CodeSigning.cs
--- a/Engine/Source/Programs/IOS/iPhonePackager/CodeSigning.cs
+++ b/Engine/Source/Programs/IOS/iPhonePackager/CodeSigning.cs
@@ -151,16 +151,6 @@ namespace iPhonePackager
 			return true;
 		}
 
-		public static bool DoRequiredFilesExist()
-		{
-			MobileProvision Provision;
-			X509Certificate2 Cert;
-			bool bOverridesExists;
-			CodeSignatureBuilder.FindRequiredFiles(out Provision, out Cert, out bOverridesExists);
-
-			return bOverridesExists && (Provision != null) && (Cert != null);
-		}
-
 		protected virtual byte[] GetMobileProvision(string CFBundleIdentifier)
 		{
 			// find the movile provision file in the library
@@ -204,6 +194,8 @@ namespace iPhonePackager
 		/// </summary>
 		public static X509Certificate2 FindCertificate(MobileProvision ProvisionToWorkFrom)
 		{
+			Program.LogVerbose("  Looking for a certificate that matches the application identifier '{0}'", ProvisionToWorkFrom.ApplicationIdentifier);
+
 			// Open the personal certificate store on this machine
 			X509Store Store = new X509Store();
 			Store.Open(OpenFlags.ReadOnly);
@@ -214,15 +206,41 @@ namespace iPhonePackager
 			{
 				X509Certificate2Collection FoundCerts = Store.Certificates.Find(X509FindType.FindBySerialNumber, SourceCert.SerialNumber, false);
 
-				if (FoundCerts.Count > 0)
+				Program.LogVerbose("  .. Provision entry SN '{0}' matched {1} installed certificate(s)", SourceCert.SerialNumber, FoundCerts.Count);
+
+				X509Certificate2 ValidInTimeCert = null;
+				foreach (X509Certificate2 TestCert in FoundCerts)
+				{
+					//@TODO: Pretty sure the certificate information from the library is in local time, not UTC and this works as expected, but it should be verified!
+					DateTime EffectiveDate = TestCert.NotBefore;
+					DateTime ExpirationDate = TestCert.NotAfter;
+					DateTime Now = DateTime.Now;
+
+					bool bCertTimeIsValid = (EffectiveDate < Now) && (ExpirationDate > Now);
+
+					Program.LogVerbose("  .. .. Installed certificate '{0}' is {1} (range '{2}' to '{3}')", TestCert.FriendlyName, bCertTimeIsValid ? "valid (choosing it)" : "EXPIRED", TestCert.GetEffectiveDateString(), TestCert.GetExpirationDateString());
+					if (bCertTimeIsValid)
+					{
+						ValidInTimeCert = TestCert;
+						break;
+					}
+				}
+
+				if (ValidInTimeCert != null)
 				{
-					Result = FoundCerts[0];
+					// Found a cert in the valid time range, quit now!
+					Result = ValidInTimeCert;
 					break;
 				}
 			}
 
 			Store.Close();
 
+			if (Result == null)
+			{
+				Program.LogVerbose("  .. Failed to find a valid certificate that was in date");
+			}
+
 			return Result;
 		}
 
@@ -366,8 +384,11 @@ namespace iPhonePackager
 
 			// Install the Apple trust chain certs (required to do a CMS signature with full chain embedded)
 			List<string> TrustChainCertFilenames = new List<string>();
-			TrustChainCertFilenames.Add("AppleWorldwideDeveloperRelationsCA.pem");
-			TrustChainCertFilenames.Add("AppleRootCA.pem");
+
+			string CertPath = Path.GetFullPath(Config.EngineBuildDirectory);
+			TrustChainCertFilenames.Add(Path.Combine(CertPath, "AppleWorldwideDeveloperRelationsCA.pem"));
+			TrustChainCertFilenames.Add(Path.Combine(CertPath, "AppleRootCA.pem"));
+
 			InstallCertificates(TrustChainCertFilenames);
 
 			// Find and load the signing cert
@@ -381,7 +402,7 @@ namespace iPhonePackager
 			}
 			else
 			{
-				Program.Log("... Found matching certificate '{0}'", SigningCert.FriendlyName);
+				Program.Log("... Found matching certificate '{0}' (valid from {1} to {2})", SigningCert.FriendlyName, SigningCert.GetEffectiveDateString(), SigningCert.GetExpirationDateString());
 			}
 		}
 
@@ -469,6 +490,8 @@ namespace iPhonePackager
 
 			foreach (MachObjectFile Exe in FatBinary.MachObjectFiles)
 			{
+				Program.Log("... Processing one mach object (binary is {0})", FatBinary.bIsFatBinary ? "fat" : "thin");
+				
 				// Pad the memory stream with extra room to handle any possible growth in the code signing data
 				int OverSize = 1024 * 1024;
 				MemoryStream OutputExeStream = new MemoryStream(SourceExeData.Length + OverSize);
@@ -562,7 +585,7 @@ namespace iPhonePackager
 				// and signed by the signature blob
 
 				// Do an initial signature just to get the size
-				Program.Log("... Initial signature step");
+				Program.Log("... Initial signature step ({0:0.00} s elapsed so far)", (DateTime.Now - SigningTime).TotalSeconds);
 				CodeSignatureBlob.SignCodeDirectory(SigningCert, SigningTime, FinalCodeDirectoryBlob);
 
 				// Compute the size of everything, and push it into the EXE header
@@ -580,7 +603,7 @@ namespace iPhonePackager
 				CodeSigningBlobLC.PatchPositionAndSize(OutputExeContext, (uint)BlobStartPosition, (uint)BlobLength);
 
 				// Now that the executable loader command has been inserted and the appropriate section modified, compute all the hashes
-				Program.Log("... Computing hashes");
+				Program.Log("... Computing hashes ({0:0.00} s elapsed so far)", (DateTime.Now - SigningTime).TotalSeconds);
 				OutputExeContext.Flush();
 
 				// Fill out the special hashes
@@ -594,7 +617,7 @@ namespace iPhonePackager
 				FinalCodeDirectoryBlob.ComputeImageHashes(OutputExeStream.ToArray());
 
 				// And compute the final signature
-				Program.Log("... Final signature step");
+				Program.Log("... Final signature step ({0:0.00} s elapsed so far)", (DateTime.Now - SigningTime).TotalSeconds);
 				CodeSignatureBlob.SignCodeDirectory(SigningCert, SigningTime, FinalCodeDirectoryBlob);
 
 				// Generate the signing blob and place it in the output (verifying it didn't change in size)
@@ -610,8 +633,10 @@ namespace iPhonePackager
 				OutputExeContext.PopPosition();
 
 				// Truncate the data so the __LINKEDIT section extends right to the end
+				Program.Log("... Committing all edits ({0:0.00} s elapsed so far)", (DateTime.Now - SigningTime).TotalSeconds);
 				OutputExeContext.CompleteWritingAndClose();
 
+				Program.Log("... Truncating/copying final binary", DateTime.Now - SigningTime);
 				ulong DesiredExecutableLength = LinkEditSegmentLC.FileSize + LinkEditSegmentLC.FileOffset;
 				byte[] FinalExeData = OutputExeStream.ToArray();
 				if ((ulong)FinalExeData.Length < DesiredExecutableLength)
@@ -621,9 +646,11 @@ namespace iPhonePackager
 				Array.Resize(ref FinalExeData, (int)DesiredExecutableLength); //@todo: Extend the file system interface so we don't have to copy 20 MB just to truncate a few hundred bytes
 
 				// Save the patched and signed executable
-				Program.Log("Saving signed executable...");
+				Program.Log("Saving signed executable... ({0:0.00} s elapsed so far)", (DateTime.Now - SigningTime).TotalSeconds);
 				FileSystem.WriteAllBytes(CFBundleExecutable, FinalExeData);
 			}
+
+			Program.Log("Finished code signing, which took {0:0.00} s", (DateTime.Now - SigningTime).TotalSeconds);
 		}
 	}
 }
diff --git a/Engine/Source/Programs/IOS/iPhonePackager/Config.cs b/Engine/Source/Programs/IOS/iPhonePackager/Config.cs
--- a/Engine/Source/Programs/IOS/iPhonePackager/Config.cs
+++ b/Engine/Source/Programs/IOS/iPhonePackager/Config.cs
@@ -231,12 +231,6 @@ namespace iPhonePackager
 			}
 		}
 
-		/// <summary>
-		/// Whether or not to allow interactive dialogs to pop up when a traditionally non-interactive command was specified
-		/// (e.g., can we pop up a configuration dialog during PackageIPA?)
-		/// </summary>
-		public static bool bAllowInteractiveDialogsDuringNonInteractiveCommands = false;
-
 		/// <summary>
 		/// Whether or not to output extra information (like every file copy and date/time stamp)
 		/// </summary>
diff --git a/Engine/Source/Programs/IOS/iPhonePackager/CookTime.cs b/Engine/Source/Programs/IOS/iPhonePackager/CookTime.cs
--- a/Engine/Source/Programs/IOS/iPhonePackager/CookTime.cs
+++ b/Engine/Source/Programs/IOS/iPhonePackager/CookTime.cs
@@ -358,7 +358,8 @@ namespace iPhonePackager
 			}
 
 			// Save the Zip
-			Program.Log("Compressing files into IPA.{0}", Config.bVerbose ? "" : "  Only large files will be listed next, but other files are also being packaged.");
+
+			Program.Log("Compressing files into IPA (-compress={1}).{0}", Config.bVerbose ? "" : "  Only large files will be listed next, but other files are also being packaged.", Config.RecompressionSetting);
 			FileSystem.Close();
 
 			TimeSpan ZipLength = DateTime.Now - StartTime;
@@ -366,7 +367,7 @@ namespace iPhonePackager
 			FileInfo FinalZipInfo = new FileInfo(Zip.Name);
 
 
-			Program.Log(String.Format("Finished repackaging into {2:0.00} MB IPA, written to '{0}' (took {1:0.00} s)",
+			Program.Log(String.Format("Finished repackaging into {2:0.00} MB IPA, written to '{0}' (took {1:0.00} s for all steps)",
 				Zip.Name,
 				ZipLength.TotalSeconds,
 				FinalZipInfo.Length / (1024.0f * 1024.0f)));
diff --git a/Engine/Source/Programs/IOS/iPhonePackager/MobileProvisionUtilities.cs b/Engine/Source/Programs/IOS/iPhonePackager/MobileProvisionUtilities.cs
--- a/Engine/Source/Programs/IOS/iPhonePackager/MobileProvisionUtilities.cs
+++ b/Engine/Source/Programs/IOS/iPhonePackager/MobileProvisionUtilities.cs
@@ -10,6 +10,7 @@ using System.Security.Cryptography;
 using System.IO;
 using System.Diagnostics;
 using System.Xml;
+using System.Globalization;
 
 namespace iPhonePackager
 {
@@ -27,6 +28,8 @@ namespace iPhonePackager
 		public string ProvisionName;
 		public bool bDebug;
 		public Utilities.PListHelper Data;
+		public DateTime CreationDate;
+		public DateTime ExpirationDate;
 
 		public static string FindCompatibleProvision(string CFBundleIdentifier, bool bCheckCert = true)
 		{
@@ -89,36 +92,80 @@ namespace iPhonePackager
 				ProvisionLibrary.Add(Provision, p);
 			}
 
+			Program.Log("Searching for mobile provisions that match the game '{0}' with CFBundleIdentifier='{1}' in '{2}'", GameName, CFBundleIdentifier, Config.ProvisionDirectory);
+
 			// check the cache for a provision matching the app id (com.company.Game)
-			foreach (KeyValuePair<string, MobileProvision> Pair in ProvisionLibrary)
+			// First checking for a contains match and then for a wildcard match
+			for (int Phase = 0; Phase < 2; ++Phase)
 			{
-				if (Pair.Value.ApplicationIdentifier.Contains(CFBundleIdentifier) && (!Config.bForDistribution || (Pair.Value.ProvisionedDeviceIDs.Count == 0 && !Pair.Value.bDebug)))
+				foreach (KeyValuePair<string, MobileProvision> Pair in ProvisionLibrary)
 				{
-					// check to see if we have a certificate for this provision
-					if (!bCheckCert || CodeSignatureBuilder.FindCertificate(Pair.Value) != null)
+					string DebugName = Path.GetFileName(Pair.Key);
+					MobileProvision TestProvision = Pair.Value;
+
+					Program.LogVerbose("  Phase {0} considering provision '{1}' named '{2}'", Phase, DebugName, TestProvision.ProvisionName);
+
+					// Validate the name
+					bool bPassesNameCheck = false;
+					if (Phase == 0)
 					{
-						return Pair.Key;
+						bPassesNameCheck = TestProvision.ApplicationIdentifier.Contains(CFBundleIdentifier);
+					}
+					else
+					{
+						bPassesNameCheck = TestProvision.ProvisionName.Contains("Wildcard") || TestProvision.ApplicationIdentifier.Contains("*");
+					}
+
+					if (!bPassesNameCheck)
+					{
+						Program.LogVerbose("  .. Failed phase {0} name check (provision app ID was {1})", Phase, TestProvision.ApplicationIdentifier);
+						continue;
+					}
+
+					bool bPassesDebugCheck = (!Config.bForDistribution || ((TestProvision.ProvisionedDeviceIDs.Count == 0) && !TestProvision.bDebug));
+					if (!bPassesDebugCheck)
+					{
+						Program.LogVerbose("  .. Failed debugging check (mode={0}, get-task-allow={1}, #devices={2})", Config.bForDistribution, TestProvision.bDebug, TestProvision.ProvisionedDeviceIDs.Count);
+						continue;
+					}
+
+					// Check to see if the provision is in date
+					DateTime CurrentUTCTime = DateTime.UtcNow;
+					bool bPassesDateCheck = (CurrentUTCTime >= TestProvision.CreationDate) && (CurrentUTCTime < TestProvision.ExpirationDate);
+					if (!bPassesDateCheck)
+					{
+						Program.LogVerbose("  .. Failed time period check (valid from {0} to {1}, but UTC time is now {2})", TestProvision.CreationDate, TestProvision.ExpirationDate, CurrentUTCTime);
+						continue;
 					}
-				}
-			}
 
-			// check the cache for a provision matching the wild card
-			foreach (KeyValuePair<string, MobileProvision> Pair in ProvisionLibrary)
-			{
-				if ((Pair.Value.ProvisionName.Contains("Wildcard") || Pair.Value.ApplicationIdentifier.Contains("*")) && (!Config.bForDistribution || (Pair.Value.ProvisionedDeviceIDs.Count == 0 && !Pair.Value.bDebug)))
-				{
 					// check to see if we have a certificate for this provision
-					if (!bCheckCert || CodeSignatureBuilder.FindCertificate(Pair.Value) != null)
+					bool bPassesHasMatchingCertCheck = false;
+					if (bCheckCert)
+					{
+						bPassesHasMatchingCertCheck = (CodeSignatureBuilder.FindCertificate(TestProvision) != null);
+					}
+					else
+					{
+						bPassesHasMatchingCertCheck = true;
+					}
+
+					if (!bPassesHasMatchingCertCheck)
 					{
-						return Pair.Key;
+						Program.LogVerbose("  .. Failed to find a matching certificate that was in date");
+						continue;
 					}
+
+					// Made it past all the tests
+					Program.LogVerbose("  Picked '{0}' with AppID '{1}' and Name '{2}' as a matching provision for the game '{3}'", DebugName, TestProvision.ApplicationIdentifier, TestProvision.ProvisionName, GameName);
+					return Pair.Key;
 				}
 			}
 
 			// check to see if there is already an embedded provision
-			string MobileProvisionFilename = Path.Combine(Config.RepackageStagingDirectory, "embedded.mobileprovision");
+			string EmbeddedMobileProvisionFilename = Path.Combine(Config.RepackageStagingDirectory, "embedded.mobileprovision");
 
-			return MobileProvisionFilename;
+			Program.Warning("Failed to find a valid matching mobile provision, will attempt to use the embedded mobile provision instead if present");
+			return EmbeddedMobileProvisionFilename;
 		}
 
 		/// <summary>
@@ -187,6 +234,21 @@ namespace iPhonePackager
 				ApplicationIdentifierPrefix = PrefixList[0];
 			}
 
+			// Example date string from the XML: "2014-06-30T20:45:55Z";
+			DateTimeStyles AppleDateStyle = DateTimeStyles.AllowWhiteSpaces | DateTimeStyles.AssumeUniversal;
+
+			string CreationDateString;
+			if (Data.GetDate("CreationDate", out CreationDateString))
+			{
+				CreationDate = DateTime.Parse(CreationDateString, CultureInfo.InvariantCulture, AppleDateStyle);
+			}
+
+			string ExpirationDateString;
+			if (Data.GetDate("ExpirationDate", out ExpirationDateString))
+			{
+				ExpirationDate = DateTime.Parse(ExpirationDateString, CultureInfo.InvariantCulture, AppleDateStyle);
+			}
+
 			// Key: DeveloperCertificates, Array<Data> (uuencoded)
 			string CertificatePassword = "";
 			List<string> CertificateList = Data.GetArray("DeveloperCertificates", "data");
diff --git a/Engine/Source/Programs/IOS/iPhonePackager/Utilities.cs b/Engine/Source/Programs/IOS/iPhonePackager/Utilities.cs
--- a/Engine/Source/Programs/IOS/iPhonePackager/Utilities.cs
+++ b/Engine/Source/Programs/IOS/iPhonePackager/Utilities.cs
@@ -351,6 +351,21 @@ namespace iPhonePackager
 				return true;
 			}
 
+			public bool GetDate(string Key, out string Value)
+			{
+				string PathToValue = String.Format("/plist/dict/key[.='{0}']/following-sibling::date[1]", Key);
+
+				XmlNode ValueNode = Doc.DocumentElement.SelectSingleNode(PathToValue);
+				if (ValueNode == null)
+				{
+					Value = "";
+					return false;
+				}
+
+				Value = ValueNode.InnerText;
+				return true;
+			}
+
 			public bool GetBool(string Key)
 			{
 				string PathToValue = String.Format("/plist/dict/key[.='{0}']/following-sibling::node()", Key);
@@ -448,8 +463,6 @@ namespace iPhonePackager
 			/// </summary>
 			public bool HasKey(string KeyName)
 			{
-				XmlNode DictionaryNode = Doc.DocumentElement.SelectSingleNode("/plist/dict");
-
 				string PathToKey = String.Format("/plist/dict/key[.='{0}']", KeyName);
 
 				XmlNode KeyNode = Doc.DocumentElement.SelectSingleNode(PathToKey);
diff --git a/Engine/Source/Programs/IOS/iPhonePackager/iPhonePackager.cs b/Engine/Source/Programs/IOS/iPhonePackager/iPhonePackager.cs
--- a/Engine/Source/Programs/IOS/iPhonePackager/iPhonePackager.cs
+++ b/Engine/Source/Programs/IOS/iPhonePackager/iPhonePackager.cs
@@ -93,6 +93,20 @@ namespace iPhonePackager
 			Log(String.Format(Line, Args));
 		}
 
+		static public void LogVerbose(string Line)
+		{
+			if (Config.bVerbose)
+			{
+				UpdateStatus(Line);
+				Console.WriteLine(Line);
+			}
+		}
+
+		static public void LogVerbose(string Line, params object[] Args)
+		{
+			LogVerbose(String.Format(Line, Args));
+		}
+
 		static public void Error( string Line, int Code = 1 )
 		{
 			if (Program.ReturnCode == 0)
@@ -175,9 +189,6 @@ namespace iPhonePackager
 						// Behavior switches
 						switch (Arg)
 						{
-							case "-interactive":
-								Config.bAllowInteractiveDialogsDuringNonInteractiveCommands = true;
-								break;
 							case "-verbose":
 								Config.bVerbose = true;
 								break;
@@ -325,25 +336,6 @@ namespace iPhonePackager
 				return false;
 			}
 
-			// If -sign was specified, check to see if the user has configured yet.  If not, we will pop up a configuration dialog
-			if (Config.bPerformResignWhenRepackaging && Config.bAllowInteractiveDialogsDuringNonInteractiveCommands)
-			{
-				string CWD = Directory.GetCurrentDirectory();
-				if (!CodeSignatureBuilder.DoRequiredFilesExist())
-				{
-					StartVisuals();
-
-					Form DisplayForm = new ToolsHub();
-					Application.Run(DisplayForm);
-					if (DisplayForm.DialogResult != DialogResult.OK)
-					{
-						Error("One or more files necessary for packaging are missing and configuration was cancelled");
-						return false;
-					}
-				}
-				Directory.SetCurrentDirectory(CWD);
-			}
-
 			return true;
 		}
 
