update to IPP to check for provision and cert on Mac
TTP342934
#ios
#mac

diff --git a/Engine/Source/Developer/iOS/IOSTargetPlatform/Private/IOSTargetPlatform.cpp b/Engine/Source/Developer/iOS/IOSTargetPlatform/Private/IOSTargetPlatform.cpp
--- a/Engine/Source/Developer/iOS/IOSTargetPlatform/Private/IOSTargetPlatform.cpp
+++ b/Engine/Source/Developer/iOS/IOSTargetPlatform/Private/IOSTargetPlatform.cpp
@@ -103,7 +103,7 @@ static FString OutputMessage;
 static void OnOutput(FString Message)
 {
 	OutputMessage += Message;
-	UE_LOG(LogTemp, Display, TEXT("%s"), *Message);
+	UE_LOG(LogTemp, Display, TEXT("%s\n"), *Message);
 }
 
 int FIOSTargetPlatform::DoesntHaveRequirements(const FString& ProjectPath, bool bProjectHasCode, FString& OutTutorialPath) const
@@ -115,6 +115,7 @@ int FIOSTargetPlatform::DoesntHaveRequirements(const FString& ProjectPath, bool
 	}
 #if PLATFORM_MAC
 	OutTutorialPath = FString("/Engine/Tutorial/Installation/InstallingXCodeTutorial.InstallingXCodeTutorial");
+    // shell to certtool
 #else
 	if (bProjectHasCode && FRocketSupport::IsRocket())
 	{
@@ -126,10 +127,18 @@ int FIOSTargetPlatform::DoesntHaveRequirements(const FString& ProjectPath, bool
 		OutTutorialPath = FString("/Engine/Tutorial/Mobile/iOSonPCValidPlugins.iOSonPCValidPlugins");
 		bReadyToBuild |= ETargetPlatformReadyStatus::PluginsUnsupported;
 	}
+#endif
 
 	// shell to IPP and get the status of the provision and cert
+#if PLATFORM_MAC
+    FString CmdExe = TEXT("/bin/sh");
+    FString ScriptPath = FPaths::ConvertRelativePathToFull(FPaths::EngineDir() / TEXT("Build/BatchFiles/Mac/RunMono.sh"));
+    FString IPPPath = FPaths::ConvertRelativePathToFull(FPaths::EngineDir() / TEXT("Binaries/DotNet/IOS/IPhonePackager.exe"));
+    FString CommandLine = FString::Printf(TEXT("\"%s\" \"%s\" Validate Engine -project \"%s\""), *ScriptPath, *IPPPath, *ProjectPath);
+#else
 	FString CmdExe = FPaths::ConvertRelativePathToFull(FPaths::EngineDir() / TEXT("Binaries/DotNet/IOS/IPhonePackager.exe"));
 	FString CommandLine = FString::Printf(TEXT("Validate Engine -project \"%s\""), *ProjectPath);
+#endif
 	TSharedPtr<FMonitoredProcess> IPPProcess = MakeShareable(new FMonitoredProcess(CmdExe, CommandLine, true));
 	OutputMessage = TEXT("");
 	IPPProcess->OnOutput().BindStatic(&OnOutput);
@@ -139,6 +148,7 @@ int FIOSTargetPlatform::DoesntHaveRequirements(const FString& ProjectPath, bool
 		FPlatformProcess::Sleep(0.01f);
 	}
 	int RetCode = IPPProcess->GetReturnCode();
+    UE_LOG(LogTemp, Display, TEXT("%s"), *OutputMessage);
 	if (RetCode == 14)
 	{
 		OutTutorialPath = FString("/Engine/Tutorial/Mobile/CreatingInfoPlist.CreatingInfoPlist");
@@ -160,7 +170,7 @@ int FIOSTargetPlatform::DoesntHaveRequirements(const FString& ProjectPath, bool
 		OutTutorialPath = FString("/Engine/Tutorial/Mobile/CreatingSigningCertAndProvisionTutorial.CreatingSigningCertAndProvisionTutorial");
 		bReadyToBuild |= ETargetPlatformReadyStatus::ProvisionNotFound;
 	}
-#endif
+
 	return bReadyToBuild;
 }
 
diff --git a/Engine/Source/Programs/IOS/iPhonePackager/CodeSigning.cs b/Engine/Source/Programs/IOS/iPhonePackager/CodeSigning.cs
--- a/Engine/Source/Programs/IOS/iPhonePackager/CodeSigning.cs
+++ b/Engine/Source/Programs/IOS/iPhonePackager/CodeSigning.cs
@@ -188,6 +188,14 @@ namespace iPhonePackager
 			}
 		}
 
+		static private string CertToolData = "";
+		static public void OutputReceivedCertToolProcessCall(Object Sender, DataReceivedEventArgs Line)
+		{
+			if ((Line != null) && !String.IsNullOrEmpty (Line.Data)) {
+				CertToolData += Line.Data + "\n";
+			}
+		}
+
 		/// <summary>
 		/// Tries to find a matching certificate on this machine from the the serial number of one of the
 		/// certificates in the mobile provision (the one in the mobileprovision is missing the public/private key pair)
@@ -196,46 +204,96 @@ namespace iPhonePackager
 		{
 			Program.LogVerbose("  Looking for a certificate that matches the application identifier '{0}'", ProvisionToWorkFrom.ApplicationIdentifier);
 
-			// Open the personal certificate store on this machine
-			X509Store Store = new X509Store();
-			Store.Open(OpenFlags.ReadOnly);
-
-			// Try finding a matching certificate from the serial number (the one in the mobileprovision is missing the public/private key pair)
 			X509Certificate2 Result = null;
-			foreach (X509Certificate2 SourceCert in ProvisionToWorkFrom.DeveloperCertificates)
-			{
-				X509Certificate2Collection FoundCerts = Store.Certificates.Find(X509FindType.FindBySerialNumber, SourceCert.SerialNumber, false);
 
-				Program.LogVerbose("  .. Provision entry SN '{0}' matched {1} installed certificate(s)", SourceCert.SerialNumber, FoundCerts.Count);
+			if (Environment.OSVersion.Platform == PlatformID.Unix || Environment.OSVersion.Platform == PlatformID.MacOSX) {
+				// run certtool y to get the currently installed certificates
+				CertToolData = "";
+				Process CertTool = new Process ();
+				CertTool.StartInfo.FileName = "/usr/bin/certtool";
+				CertTool.StartInfo.UseShellExecute = false;
+				CertTool.StartInfo.Arguments = "y";
+				CertTool.StartInfo.RedirectStandardOutput = true;
+				CertTool.OutputDataReceived += new DataReceivedEventHandler (OutputReceivedCertToolProcessCall);
+				CertTool.Start ();
+				CertTool.BeginOutputReadLine ();
+				CertTool.WaitForExit ();
+				if (CertTool.ExitCode == 0) {
+					foreach (X509Certificate2 SourceCert in ProvisionToWorkFrom.DeveloperCertificates) {
+						X509Certificate2 ValidInTimeCert = null;
+						// see if certificate can be found by serial number
+						string SerialNumber = SourceCert.GetSerialNumberString ();
+						for (int Index = 14; Index > 0; Index -= 2) {
+							SerialNumber = SerialNumber.Insert (Index, " ");
+						}
 
-				X509Certificate2 ValidInTimeCert = null;
-				foreach (X509Certificate2 TestCert in FoundCerts)
-				{
-					//@TODO: Pretty sure the certificate information from the library is in local time, not UTC and this works as expected, but it should be verified!
-					DateTime EffectiveDate = TestCert.NotBefore;
-					DateTime ExpirationDate = TestCert.NotAfter;
-					DateTime Now = DateTime.Now;
+						if (CertToolData.Contains (SerialNumber)) {
+							// check the cert time
+							int StartLine = CertToolData.IndexOf (SerialNumber);
+							string BeforeStart = CertToolData.Substring( CertToolData.IndexOf( ":", CertToolData.IndexOf ("Not Before", StartLine))+1);
+							BeforeStart = BeforeStart.Remove(BeforeStart.IndexOf("\n"));
+							string AfterStart = CertToolData.Substring( CertToolData.IndexOf( ":", CertToolData.IndexOf ("Not After", StartLine))+1);
+							AfterStart = AfterStart.Remove(AfterStart.IndexOf("\n"));
+							string CommonName = CertToolData.Substring( CertToolData.IndexOf( ":", CertToolData.IndexOf ("Common Name", CertToolData.IndexOf("Subject Name", StartLine)))+1);
+							CommonName = CommonName.Remove(CommonName.IndexOf("\n"));
+
+							DateTime EffectiveDate = DateTime.Parse (BeforeStart);
+							DateTime ExpirationDate = DateTime.Parse (AfterStart);
+							DateTime Now = DateTime.Now;
+
+							bool bCertTimeIsValid = (EffectiveDate < Now) && (ExpirationDate > Now);
+
+							Program.LogVerbose ("  .. .. Installed certificate '{0}' is {1} (range '{2}' to '{3}')", CommonName, bCertTimeIsValid ? "valid (choosing it)" : "EXPIRED", BeforeStart, AfterStart);
+							if (bCertTimeIsValid) {
+								ValidInTimeCert = SourceCert;
+							}
+						}
 
-					bool bCertTimeIsValid = (EffectiveDate < Now) && (ExpirationDate > Now);
+						if (ValidInTimeCert != null) {
+							// Found a cert in the valid time range, quit now!
+							Result = ValidInTimeCert;
+							break;
+						}
+					}
+				} else {
+				}
+			} else {
+				// Open the personal certificate store on this machine
+				X509Store Store = new X509Store (StoreName.TrustedPeople, StoreLocation.LocalMachine);
+				Store.Open (OpenFlags.ReadOnly);
+
+				// Try finding a matching certificate from the serial number (the one in the mobileprovision is missing the public/private key pair)
+				foreach (X509Certificate2 SourceCert in ProvisionToWorkFrom.DeveloperCertificates) {
+					X509Certificate2Collection FoundCerts = Store.Certificates.Find (X509FindType.FindBySerialNumber, SourceCert.SerialNumber, false);
+
+					Program.LogVerbose ("  .. Provision entry SN '{0}' matched {1} installed certificate(s)", SourceCert.SerialNumber, FoundCerts.Count);
+
+					X509Certificate2 ValidInTimeCert = null;
+					foreach (X509Certificate2 TestCert in FoundCerts) {
+						//@TODO: Pretty sure the certificate information from the library is in local time, not UTC and this works as expected, but it should be verified!
+						DateTime EffectiveDate = TestCert.NotBefore;
+						DateTime ExpirationDate = TestCert.NotAfter;
+						DateTime Now = DateTime.Now;
+
+						bool bCertTimeIsValid = (EffectiveDate < Now) && (ExpirationDate > Now);
+
+						Program.LogVerbose ("  .. .. Installed certificate '{0}' is {1} (range '{2}' to '{3}')", TestCert.FriendlyName, bCertTimeIsValid ? "valid (choosing it)" : "EXPIRED", TestCert.GetEffectiveDateString (), TestCert.GetExpirationDateString ());
+						if (bCertTimeIsValid) {
+							ValidInTimeCert = TestCert;
+							break;
+						}
+					}
 
-					Program.LogVerbose("  .. .. Installed certificate '{0}' is {1} (range '{2}' to '{3}')", TestCert.FriendlyName, bCertTimeIsValid ? "valid (choosing it)" : "EXPIRED", TestCert.GetEffectiveDateString(), TestCert.GetExpirationDateString());
-					if (bCertTimeIsValid)
-					{
-						ValidInTimeCert = TestCert;
+					if (ValidInTimeCert != null) {
+						// Found a cert in the valid time range, quit now!
+						Result = ValidInTimeCert;
 						break;
 					}
 				}
 
-				if (ValidInTimeCert != null)
-				{
-					// Found a cert in the valid time range, quit now!
-					Result = ValidInTimeCert;
-					break;
-				}
+				Store.Close ();
 			}
 
-			Store.Close();
-
 			if (Result == null)
 			{
 				Program.LogVerbose("  .. Failed to find a valid certificate that was in date");
diff --git a/Engine/Source/Programs/IOS/iPhonePackager/Config.cs b/Engine/Source/Programs/IOS/iPhonePackager/Config.cs
--- a/Engine/Source/Programs/IOS/iPhonePackager/Config.cs
+++ b/Engine/Source/Programs/IOS/iPhonePackager/Config.cs
@@ -69,7 +69,13 @@ namespace iPhonePackager
 		/// </summary>
 		public static string ProvisionDirectory
 		{
-			get { return Environment.GetFolderPath(Environment.SpecialFolder.LocalApplicationData) + "/Apple Computer/MobileDevice/Provisioning Profiles/"; }
+			get {
+				if (Environment.OSVersion.Platform == PlatformID.MacOSX || Environment.OSVersion.Platform == PlatformID.Unix) {
+					return Environment.GetEnvironmentVariable("HOME") + "/Library/MobileDevice/Provisioning Profiles/";
+				} else {
+					return Environment.GetFolderPath (Environment.SpecialFolder.LocalApplicationData) + "/Apple Computer/MobileDevice/Provisioning Profiles/";
+				}
+			}
 		}
 
 		/// <summary>
@@ -77,7 +83,20 @@ namespace iPhonePackager
 		/// </summary>
 		public static string EngineBuildDirectory
 		{
-			get { return Path.GetFullPath(RootRelativePath + @"Engine\Build\IOS"); }
+			get { 
+				if (Environment.OSVersion.Platform == PlatformID.MacOSX || Environment.OSVersion.Platform == PlatformID.Unix) {
+					return Path.GetFullPath (RootRelativePath + "Engine/Build/IOS"); 
+				} else {
+					return Path.GetFullPath (RootRelativePath + @"Engine\Build\IOS"); 
+				}
+			}
+		}
+
+		static Config()
+		{
+			if (Environment.OSVersion.Platform == PlatformID.MacOSX || Environment.OSVersion.Platform == PlatformID.Unix) {
+				RootRelativePath = "../../../../";
+			}
 		}
 
 		/// <summary>
@@ -337,7 +356,11 @@ namespace iPhonePackager
 			// else we assume old school game name and look for it
 			else
 			{
-				GameDirectory = Path.GetFullPath(Path.Combine(Config.RootRelativePath, Program.GameName));
+				if (Program.GameName == "UE4Game") {
+					GameDirectory = Path.GetFullPath (Path.Combine (Config.RootRelativePath, GamePath));
+				} else {
+					GameDirectory = Path.GetFullPath (Path.Combine (Config.RootRelativePath, Program.GameName));
+				}
 			}
 
 			if (!Directory.Exists(GameDirectory))
diff --git a/Engine/Source/Programs/IOS/iPhonePackager/iPhonePackager.cs b/Engine/Source/Programs/IOS/iPhonePackager/iPhonePackager.cs
--- a/Engine/Source/Programs/IOS/iPhonePackager/iPhonePackager.cs
+++ b/Engine/Source/Programs/IOS/iPhonePackager/iPhonePackager.cs
@@ -598,7 +598,15 @@ namespace iPhonePackager
 				{
 					case "validate":
 						// check to see if iTunes is installed
-						string dllPath = Microsoft.Win32.Registry.GetValue("HKEY_LOCAL_MACHINE\\SOFTWARE\\Wow6432Node\\Apple Inc.\\Apple Mobile Device Support\\Shared", "iTunesMobileDeviceDLL", null) as string; 
+						string dllPath = "";
+						if (Environment.OSVersion.Platform == PlatformID.MacOSX || Environment.OSVersion.Platform == PlatformID.Unix)
+						{
+							dllPath = "/Applications/Xcode.app/Contents/MacOS/Xcode";
+						}
+						else
+						{
+							dllPath = Microsoft.Win32.Registry.GetValue("HKEY_LOCAL_MACHINE\\SOFTWARE\\Wow6432Node\\Apple Inc.\\Apple Mobile Device Support\\Shared", "iTunesMobileDeviceDLL", null) as string; 
+						}
 						if (String.IsNullOrEmpty(dllPath) || !File.Exists(dllPath))
 						{
 							Error("iTunes Not Found!!", (int)ErrorCodes.Error_SDKNotFound);
