ios certificate fix

diff --git a/Engine/Source/Programs/IOS/iPhonePackager/CodeSigning.cs b/Engine/Source/Programs/IOS/iPhonePackager/CodeSigning.cs
--- a/Engine/Source/Programs/IOS/iPhonePackager/CodeSigning.cs
+++ b/Engine/Source/Programs/IOS/iPhonePackager/CodeSigning.cs
@@ -259,7 +259,7 @@ namespace iPhonePackager
 				}
 			} else {
 				// Open the personal certificate store on this machine
-				X509Store Store = new X509Store (StoreName.TrustedPeople, StoreLocation.LocalMachine);
+				X509Store Store = new X509Store ();
 				Store.Open (OpenFlags.ReadOnly);
 
 				// Try finding a matching certificate from the serial number (the one in the mobileprovision is missing the public/private key pair)
