Adding the capacity to specify additional commandline flags from session front end for deployment on IOS.

diff --git a/Engine/Source/Programs/AutomationTool/IOS/IOSPlatform.Automation.cs b/Engine/Source/Programs/AutomationTool/IOS/IOSPlatform.Automation.cs
--- a/Engine/Source/Programs/AutomationTool/IOS/IOSPlatform.Automation.cs
+++ b/Engine/Source/Programs/AutomationTool/IOS/IOSPlatform.Automation.cs
@@ -818,7 +818,7 @@ public class IOSPlatform : Platform
 			var IPPExe = CombinePaths(CmdEnv.LocalRoot, "Engine/Binaries/DotNET/IOS/IPhonePackager.exe");
 
 			// check for it in the stage directory
-			RunAndLog(CmdEnv, IPPExe, "Deploy \"" + Path.GetFullPath(StagedIPA) + "\"" + (String.IsNullOrEmpty(Params.Device) ? "" : " -device " + Params.Device.Substring(4)));
+			RunAndLog(CmdEnv, IPPExe, "Deploy \"" + Path.GetFullPath(StagedIPA) + "\"" + (String.IsNullOrEmpty(Params.Device) ? "" : " -device " + Params.Device.Substring(4)) + (" -additionalcommandline " + "\"" + Params.RunCommandline + "\""));
 		}
         PrintRunTime();
     }
diff --git a/Engine/Source/Programs/IOS/iPhonePackager/DeployTime.cs b/Engine/Source/Programs/IOS/iPhonePackager/DeployTime.cs
--- a/Engine/Source/Programs/IOS/iPhonePackager/DeployTime.cs
+++ b/Engine/Source/Programs/IOS/iPhonePackager/DeployTime.cs
@@ -137,8 +137,62 @@ namespace iPhonePackager
 				case "install":
 					{
 						string IPAPath = GamePath;
+						string AdditionalCommandline = Program.AdditionalCommandline;
 
-						if (File.Exists(GamePath))
+						if (!String.IsNullOrEmpty(AdditionalCommandline))
+                        {
+                            // Read the mobile provision to check for issues
+                            FileOperations.ReadOnlyZipFileSystem Zip = new FileOperations.ReadOnlyZipFileSystem(IPAPath);
+                            try
+                            {
+                                // Compare the commandline embedded to prevent us from any unnecessary writing.
+                                byte[] CommandlineBytes = Zip.ReadAllBytes("ue4commandline.txt");
+                                string ExistingCommandline = Encoding.UTF8.GetString(CommandlineBytes, 0, CommandlineBytes.Length);
+								if (ExistingCommandline != AdditionalCommandline)
+                                {
+                                    // Ensure we have a temp dir to stage our temporary ipa
+                                    if( !Directory.Exists( Config.PCStagingRootDir ) )
+                                    {
+                                        Directory.CreateDirectory(Config.PCStagingRootDir);
+                                    }
+
+                                    string TmpFilePath = Path.Combine(Path.GetDirectoryName(Config.PCStagingRootDir), Path.GetFileNameWithoutExtension(IPAPath) + ".tmp.ipa");
+                                    if( File.Exists( TmpFilePath ) )
+                                    {
+                                        File.Delete(TmpFilePath);
+                                    }
+
+                                    File.Copy(IPAPath, TmpFilePath);
+                                
+									// Get the project name:
+									string ProjectFile = ExistingCommandline.Split(' ').FirstOrDefault();
+
+                                    // Write out the new commandline.
+                                    FileOperations.ZipFileSystem WritableZip = new FileOperations.ZipFileSystem(TmpFilePath);
+									byte[] NewCommandline = Encoding.UTF8.GetBytes(ProjectFile + " " + AdditionalCommandline);
+                                    WritableZip.WriteAllBytes("ue4commandline.txt", NewCommandline);
+
+                                    // We need to residn the application after the commandline file has changed.
+                                    CodeSignatureBuilder CodeSigner = new CodeSignatureBuilder();
+                                    CodeSigner.FileSystem = WritableZip;
+
+                                    CodeSigner.PrepareForSigning();
+                                    CodeSigner.PerformSigning();
+
+                                    WritableZip.Close();
+
+                                    // Set the deploying ipa path to our new ipa
+                                    IPAPath = TmpFilePath;
+                                }
+                            }
+                            catch (System.Exception ex)
+                            {
+                                Program.Warning(String.Format("Failed to override the commandline.txt file: ({0})", ex.Message));
+                            }
+                            Zip.Close();
+                        }
+
+						if (File.Exists(IPAPath))
 						{
 							if (!String.IsNullOrEmpty(Config.DeviceId) && !Config.DeviceId.Contains("All_iOS_On"))
 							{
diff --git a/Engine/Source/Programs/IOS/iPhonePackager/MobileProvisionUtilities.cs b/Engine/Source/Programs/IOS/iPhonePackager/MobileProvisionUtilities.cs
--- a/Engine/Source/Programs/IOS/iPhonePackager/MobileProvisionUtilities.cs
+++ b/Engine/Source/Programs/IOS/iPhonePackager/MobileProvisionUtilities.cs
@@ -53,7 +53,7 @@ namespace iPhonePackager
 
 			#region remove after we provide an install mechanism
 			// copy all of the provisions from the game directory to the library
-
+			if (!String.IsNullOrEmpty(Config.ProjectFile))
 			{
 				var ProjectFileBuildIOSPath = Path.GetDirectoryName(Config.ProjectFile) + "/Build/IOS/";
 				if (Directory.Exists(ProjectFileBuildIOSPath))
diff --git a/Engine/Source/Programs/IOS/iPhonePackager/iPhonePackager.cs b/Engine/Source/Programs/IOS/iPhonePackager/iPhonePackager.cs
--- a/Engine/Source/Programs/IOS/iPhonePackager/iPhonePackager.cs
+++ b/Engine/Source/Programs/IOS/iPhonePackager/iPhonePackager.cs
@@ -78,6 +78,7 @@ namespace iPhonePackager
 		static public SlowProgressDialog ProgressDialog = null;
 		static public BackgroundWorker BGWorker = null;
 		static public int ProgressIndex = 0;
+		static public string AdditionalCommandline = "";
 
 		static public void UpdateStatus(string Line)
 		{
@@ -309,6 +310,18 @@ namespace iPhonePackager
 								return false;
 							}
 						}
+						else if (Arg == "-additionalcommandline")
+						{
+							// make sure there's at least one more arg
+							if (Arguments.Length > ArgIndex + 1)
+							{
+								AdditionalCommandline = Arguments[++ArgIndex];
+							}
+							else
+							{
+								return false;
+							}
+						}
 					}
 					else
 					{
