fix for parsing UTC time from mobile provision files
#ios

diff --git a/Engine/Source/Programs/IOS/iPhonePackager/MobileProvisionUtilities.cs b/Engine/Source/Programs/IOS/iPhonePackager/MobileProvisionUtilities.cs
--- a/Engine/Source/Programs/IOS/iPhonePackager/MobileProvisionUtilities.cs
+++ b/Engine/Source/Programs/IOS/iPhonePackager/MobileProvisionUtilities.cs
@@ -239,7 +239,7 @@ namespace iPhonePackager
 			}
 
 			// Example date string from the XML: "2014-06-30T20:45:55Z";
-			DateTimeStyles AppleDateStyle = DateTimeStyles.AllowWhiteSpaces | DateTimeStyles.AssumeUniversal;
+			DateTimeStyles AppleDateStyle = DateTimeStyles.AllowWhiteSpaces | DateTimeStyles.AssumeUniversal | DateTimeStyles.AdjustToUniversal;
 
 			string CreationDateString;
 			if (Data.GetDate("CreationDate", out CreationDateString))
