fix for signing fat binaries on PC
UE-5018
#ios

diff --git a/Engine/Source/Programs/IOS/iPhonePackager/CodeSigning.cs b/Engine/Source/Programs/IOS/iPhonePackager/CodeSigning.cs
--- a/Engine/Source/Programs/IOS/iPhonePackager/CodeSigning.cs
+++ b/Engine/Source/Programs/IOS/iPhonePackager/CodeSigning.cs
@@ -553,6 +553,7 @@ namespace iPhonePackager
 			//@TODO: Verify it's an executable (not an object file, etc...)
 			ulong CurrentStreamOffset = 0;
 			byte[] FinalExeData = new byte[SourceExeData.Length + 1024 * 1024];
+			int ArchIndex = 0;
 
 			foreach (MachObjectFile Exe in FatBinary.MachObjectFiles)
 			{
@@ -560,11 +561,25 @@ namespace iPhonePackager
 
 				// Pad the memory stream with extra room to handle any possible growth in the code signing data
 				int OverSize = 1024 * 1024;
-				MemoryStream OutputExeStream = new MemoryStream(SourceExeData.Length - (int)CurrentStreamOffset + OverSize);
+				int ExeSize = (FatBinary.bIsFatBinary ? (int)FatBinary.Archs[ArchIndex].Size : SourceExeData.Length);
+				MemoryStream OutputExeStream = new MemoryStream(ExeSize + OverSize);
+
+				// Copy the data up to the executable into the final stream
+				if (FatBinary.bIsFatBinary)
+				{
+					OutputExeStream.Seek(0, SeekOrigin.Begin);
+					OutputExeStream.Write(SourceExeData, (int)CurrentStreamOffset, (int)FatBinary.Archs[ArchIndex].Offset - (int)CurrentStreamOffset);
+					OutputExeStream.Seek(0, SeekOrigin.Begin);
+
+					byte[] HeaderData = OutputExeStream.ToArray();
+					HeaderData.CopyTo(FinalExeData, (long)CurrentStreamOffset);
+					CurrentStreamOffset += (ulong)HeaderData.Length;
+				}
 
 				// Copy the executable into the stream
+				int ExeOffset = (FatBinary.bIsFatBinary ? (int)FatBinary.Archs[ArchIndex].Offset : 0);
 				OutputExeStream.Seek(0, SeekOrigin.Begin);
-				OutputExeStream.Write(SourceExeData, (int)CurrentStreamOffset, SourceExeData.Length - (int)CurrentStreamOffset);
+				OutputExeStream.Write(SourceExeData, ExeOffset, ExeSize);
 				OutputExeStream.Seek(0, SeekOrigin.Begin);
 				long Length = OutputExeStream.Length;
 
@@ -711,6 +726,9 @@ namespace iPhonePackager
 				byte[] Data = OutputExeStream.ToArray();
 				Data.CopyTo(FinalExeData, (long)CurrentStreamOffset);
 				CurrentStreamOffset += DesiredExecutableLength;
+
+				// increment the architecture index
+				ArchIndex++;
 			}
 
 			// resize to the finale size
