re-enable the key-pair dialog for IPP when coming from the editor
#ios

diff --git a/Engine/Source/Programs/IOS/iPhonePackager/ToolsHub.cs b/Engine/Source/Programs/IOS/iPhonePackager/ToolsHub.cs
--- a/Engine/Source/Programs/IOS/iPhonePackager/ToolsHub.cs
+++ b/Engine/Source/Programs/IOS/iPhonePackager/ToolsHub.cs
@@ -261,7 +261,7 @@ namespace iPhonePackager
 						}
 
 						// If the certificate doesn't have a private key pair, ask the user to provide one
-						if (ShowPrompt && !Cert.HasPrivateKey)
+						if (!Cert.HasPrivateKey)
 						{
 							string ErrorMsg = "Certificate does not include a private key and cannot be used to code sign";
 
