fix for message boxes showing up behind the editor
also made the IPP call async
UE-5368 & UE-5082
#ios

diff --git a/Engine/Source/Developer/iOS/IOSPlatformEditor/Private/IOSTargetSettingsCustomization.cpp b/Engine/Source/Developer/iOS/IOSPlatformEditor/Private/IOSTargetSettingsCustomization.cpp
--- a/Engine/Source/Developer/iOS/IOSPlatformEditor/Private/IOSTargetSettingsCustomization.cpp
+++ b/Engine/Source/Developer/iOS/IOSPlatformEditor/Private/IOSTargetSettingsCustomization.cpp
@@ -570,18 +570,11 @@ FReply FIOSTargetSettingsCustomization::OnInstallProvisionClicked()
 		FString CmdExe = FPaths::ConvertRelativePathToFull(FPaths::EngineDir() / TEXT("Binaries/DotNet/IOS/IPhonePackager.exe"));
 		FString CommandLine = FString::Printf(TEXT("Install Engine -project \"%s\" -provision \"%s\""), *ProjectPath, *ProvisionPath);
 #endif
-		TSharedPtr<FMonitoredProcess> IPPProcess = MakeShareable(new FMonitoredProcess(CmdExe, CommandLine, true));
+		IPPProcess = MakeShareable(new FMonitoredProcess(CmdExe, CommandLine, true));
 		OutputMessage = TEXT("");
 		IPPProcess->OnOutput().BindStatic(&OnOutput);
 		IPPProcess->Launch();
-		while(IPPProcess->IsRunning())
-		{
-			FPlatformProcess::Sleep(0.01f);
-		}
-		int RetCode = IPPProcess->GetReturnCode();
-		ensure(RetCode == 0);
-
-		UpdateStatus();
+		FTicker::GetCoreTicker().AddTicker(FTickerDelegate::CreateRaw(this, &FIOSTargetSettingsCustomization::UpdateStatusDelegate), 10.0f);
 	}
 
 	return FReply::Handled();
@@ -635,18 +628,11 @@ FReply FIOSTargetSettingsCustomization::OnInstallCertificateClicked()
 		FString CmdExe = FPaths::ConvertRelativePathToFull(FPaths::EngineDir() / TEXT("Binaries/DotNet/IOS/IPhonePackager.exe"));
 		FString CommandLine = FString::Printf(TEXT("Install Engine -project \"%s\" -certificate \"%s\""), *ProjectPath, *CertPath);
 #endif
-		TSharedPtr<FMonitoredProcess> IPPProcess = MakeShareable(new FMonitoredProcess(CmdExe, CommandLine, true));
+		IPPProcess = MakeShareable(new FMonitoredProcess(CmdExe, CommandLine, true));
 		OutputMessage = TEXT("");
 		IPPProcess->OnOutput().BindStatic(&OnOutput);
 		IPPProcess->Launch();
-		while(IPPProcess->IsRunning())
-		{
-			FPlatformProcess::Sleep(0.01f);
-		}
-		int RetCode = IPPProcess->GetReturnCode();
-		ensure(RetCode == 0);
-
-		UpdateStatus();
+		FTicker::GetCoreTicker().AddTicker(FTickerDelegate::CreateRaw(this, &FIOSTargetSettingsCustomization::UpdateStatusDelegate), 10.0f);
 	}
 
 	return FReply::Handled();
@@ -682,6 +668,18 @@ const FSlateBrush* FIOSTargetSettingsCustomization::GetCertificateStatus() const
 	}
 }
 
+bool FIOSTargetSettingsCustomization::UpdateStatusDelegate(float DeltaTime)
+{
+	if (IPPProcess->IsRunning())
+	{
+		return true;
+	}
+	int RetCode = IPPProcess->GetReturnCode();
+	ensure(RetCode == 0);
+	UpdateStatus();
+
+	return false;
+}
 //////////////////////////////////////////////////////////////////////////
 
 #undef LOCTEXT_NAMESPACE
\ No newline at end of file
diff --git a/Engine/Source/Developer/iOS/IOSPlatformEditor/Private/IOSTargetSettingsCustomization.h b/Engine/Source/Developer/iOS/IOSPlatformEditor/Private/IOSTargetSettingsCustomization.h
--- a/Engine/Source/Developer/iOS/IOSPlatformEditor/Private/IOSTargetSettingsCustomization.h
+++ b/Engine/Source/Developer/iOS/IOSPlatformEditor/Private/IOSTargetSettingsCustomization.h
@@ -37,6 +37,8 @@ private:
 	bool bProvisionInstalled;
 	bool bCertificateInstalled;
 
+	TSharedPtr<FMonitoredProcess> IPPProcess;
+
 private:
 	FIOSTargetSettingsCustomization();
 
@@ -72,4 +74,7 @@ private:
 
 	// Update the provision status
 	void UpdateStatus();
+
+	// status tick delay
+	bool UpdateStatusDelegate(float delay);
 };
diff --git a/Engine/Source/Programs/IOS/iPhonePackager/ToolsHub.cs b/Engine/Source/Programs/IOS/iPhonePackager/ToolsHub.cs
--- a/Engine/Source/Programs/IOS/iPhonePackager/ToolsHub.cs
+++ b/Engine/Source/Programs/IOS/iPhonePackager/ToolsHub.cs
@@ -11,11 +11,15 @@ using System.Windows.Forms;
 using System.IO;
 using System.Security.Cryptography.X509Certificates;
 using System.Diagnostics;
+using System.Runtime.InteropServices;
 
 namespace iPhonePackager
 {
 	public partial class ToolsHub : Form
 	{
+		[DllImport("user32.dll")]
+		public static extern int MessageBox(IntPtr hWnd, String Text, String Caption, uint type);
+
 		protected Dictionary<int, Bitmap> CheckStateImages = new Dictionary<int, Bitmap>();
 
 		public static string IpaFilter = "iOS packaged applications (*.ipa)|*.ipa|All Files (*.*)|*.*";
@@ -142,7 +146,7 @@ namespace iPhonePackager
 
 		public static void ShowError(string Message)
 		{
-			MessageBox.Show(Message, Config.AppDisplayName, MessageBoxButtons.OK, MessageBoxIcon.Error);
+			System.Windows.Forms.MessageBox.Show(Message, Config.AppDisplayName, MessageBoxButtons.OK, MessageBoxIcon.Error);
 		}
 
 		public static bool IsProfileForDistribution(MobileProvision Provision)
@@ -179,7 +183,7 @@ namespace iPhonePackager
 							OldProvision.ProvisionName,
 							Provision.ProvisionName);
 
-						if (ShowPrompt && MessageBox.Show(MessagePrompt, Config.AppDisplayName, MessageBoxButtons.YesNo, MessageBoxIcon.Question) == DialogResult.No)
+						if (ShowPrompt && System.Windows.Forms.MessageBox.Show(MessagePrompt, Config.AppDisplayName, MessageBoxButtons.YesNo, MessageBoxIcon.Question) == DialogResult.No)
 						{
 							return;
 						}
@@ -266,10 +270,9 @@ namespace iPhonePackager
 							string ErrorMsg = "Certificate does not include a private key and cannot be used to code sign";
 
 							// Prompt for a key pair
-							if (MessageBox.Show("Next, please choose the key pair that you made when generating the certificate request.",
+							if (MessageBox(new IntPtr(0), "Next, please choose the key pair that you made when generating the certificate request.",
 								Config.AppDisplayName,
-								MessageBoxButtons.OK,
-								MessageBoxIcon.Information) == DialogResult.OK)
+								0x00000000 | 0x00000040 | 0x00001000 | 0x00010000) == 1)
 							{
 								string KeyFilename;
 								if (ShowOpenFileDialog(KeysFilter, "Choose the key pair that belongs with the signing certificate", "", "", ref ChoosingFilesToInstallDirectory, out KeyFilename))
@@ -301,7 +304,8 @@ namespace iPhonePackager
 			{
 				string ErrorMsg = String.Format("Failed to load or install certificate due to an error: '{0}'", ex.Message);
 				Program.Error(ErrorMsg);
-				MessageBox.Show(ErrorMsg, Config.AppDisplayName, MessageBoxButtons.OK, MessageBoxIcon.Error);
+				System.Threading.Thread.Sleep(500);
+				MessageBox(new IntPtr(0), ErrorMsg, Config.AppDisplayName, 0x00000000 | 0x00000010 | 0x00001000 | 0x00010000);
 			}
 		}
 
