Second part of SSH changes - SSH can now be used for remote IOS builds.
- Added changes to IPP to allow SSH commands to run as opposed to RPC. This is configurable through the commandline.
- Remote Tool Chain fixes for SSH which were preventing downloads to work. Updated rsync command command
- Removed the ParseProjectSettings calls from some IOS tool chain functions as they were unnecessarily going back through the config cache.
- Added an SSH Command Helper to IPP
- UBT passes SSH details through to IPP on commandline.

diff --git a/Engine/Source/Programs/IOS/iPhonePackager/CompileTime.cs b/Engine/Source/Programs/IOS/iPhonePackager/CompileTime.cs
--- a/Engine/Source/Programs/IOS/iPhonePackager/CompileTime.cs
+++ b/Engine/Source/Programs/IOS/iPhonePackager/CompileTime.cs
@@ -3,6 +3,7 @@
  */
 
 using System;
+using System.Collections;
 using System.Collections.Generic;
 using System.Diagnostics;
 using System.IO;
@@ -97,7 +98,7 @@ namespace iPhonePackager
 
 			// MacName=%ue4.iPhone_SigningServerName%
 			MacName = Config.OverrideMacName != null ? Config.OverrideMacName : Utilities.GetEnvironmentVariable( "ue.IOSSigningServer", "a1487" );
-			iPhone_SigningDevRootMac = "/UE4/Builds";
+			iPhone_SigningDevRootMac = Config.OverrideDevRoot != null ? Config.OverrideDevRoot : "/UE4/Builds";
 
 			// get the path to mirror into on the Mac
 			string BinariesDir = Path.GetFullPath(Path.Combine(Environment.CurrentDirectory, @"..\.."));
@@ -273,8 +274,6 @@ namespace iPhonePackager
 			string WorkingFolder = "";
 			string DisplayCommandLine = "";
 
-			Program.Log( "Running RPC on " + MacName + " ... " );
-
 			switch (RPCCommand.ToLowerInvariant())
 			{
 			case "deletemacstagingfiles":
@@ -405,36 +404,49 @@ namespace iPhonePackager
 			Program.Log( " ... " + DisplayCommandLine );
 			Program.Log(" ... full command: " +  MacName + " " + CommandLine);
 
-			Process RPCUtil = new Process();
-			RPCUtil.StartInfo.FileName = @"..\RPCUtility.exe";
-			RPCUtil.StartInfo.UseShellExecute = false;
-			RPCUtil.StartInfo.Arguments = MacName + " " + CommandLine;
-			RPCUtil.StartInfo.RedirectStandardOutput = true;
-			RPCUtil.StartInfo.RedirectStandardError = true;
-			RPCUtil.OutputDataReceived += new DataReceivedEventHandler(OutputReceivedRemoteProcessCall);
-			RPCUtil.ErrorDataReceived += new DataReceivedEventHandler(OutputReceivedRemoteProcessCall);
+			bool bSuccess = false;
+			if( Config.bUseRPCUtil )
+			{
+				Program.Log( "Running RPC on " + MacName + " ... " );
 
-			RPCUtil.Start();
-			
-			RPCUtil.BeginOutputReadLine();
-			RPCUtil.BeginErrorReadLine();
+				Process RPCUtil = new Process();
+				RPCUtil.StartInfo.FileName = @"..\RPCUtility.exe";
+				RPCUtil.StartInfo.UseShellExecute = false;
+				RPCUtil.StartInfo.Arguments = MacName + " " + CommandLine;
+				RPCUtil.StartInfo.RedirectStandardOutput = true;
+				RPCUtil.StartInfo.RedirectStandardError = true;
+				RPCUtil.OutputDataReceived += new DataReceivedEventHandler(OutputReceivedRemoteProcessCall);
+				RPCUtil.ErrorDataReceived += new DataReceivedEventHandler(OutputReceivedRemoteProcessCall);
 
-			RPCUtil.WaitForExit();
+				RPCUtil.Start();
 
-			if (RPCUtil.ExitCode != 0)
-			{
-				Program.Error("RPCCommand {0} failed with return code {1}", RPCCommand, RPCUtil.ExitCode);
-				switch (RPCCommand.ToLowerInvariant())
+				RPCUtil.BeginOutputReadLine();
+				RPCUtil.BeginErrorReadLine();
+
+				RPCUtil.WaitForExit();
+
+				bSuccess = (RPCUtil.ExitCode == 0);
+				if (bSuccess == false)
 				{
-					case "installprovision":
-						Program.Error("Ensure your access permissions for '~/Library/MobileDevice/Provisioning Profiles' are set correctly.");
-						break;
-					default:
-						break;
+					Program.Error("RPCCommand {0} failed with return code {1}", RPCCommand, RPCUtil.ExitCode);
+					switch (RPCCommand.ToLowerInvariant())
+					{
+						case "installprovision":
+							Program.Error("Ensure your access permissions for '~/Library/MobileDevice/Provisioning Profiles' are set correctly.");
+							break;
+						default:
+							break;
+					}
 				}
 			}
+			else
+			{
+				Program.Log("Running SSH on " + MacName + " ... ");
+				bSuccess = SSHCommandHelper.Command(MacName, DisplayCommandLine, WorkingFolder);
+			}
+
 
-			return (RPCUtil.ExitCode == 0);
+			return bSuccess;
 		}
 
 
@@ -515,7 +527,6 @@ namespace iPhonePackager
 
 		static public void ExecuteRemoteCommand(string RemoteCommand)
 		{
-			Program.Log("Running RPC on " + MacName + " ... ");
 			RunRPCUtilty(RemoteCommand);
 		}
 
diff --git a/Engine/Source/Programs/IOS/iPhonePackager/Config.cs b/Engine/Source/Programs/IOS/iPhonePackager/Config.cs
--- a/Engine/Source/Programs/IOS/iPhonePackager/Config.cs
+++ b/Engine/Source/Programs/IOS/iPhonePackager/Config.cs
@@ -49,6 +49,16 @@ namespace iPhonePackager
 		/// </summary>
 		public static string DeviceId = "";
 
+		/// <summary>
+		/// Determine whether to use RPC Util or not
+		/// </summary>
+		public static bool bUseRPCUtil = true;
+
+		/// <summary>
+		/// Optional override for the dev root on the target mac for remote builds.
+		/// </summary>
+		public static string OverrideDevRoot = null;
+
 		/// <summary>
 		/// The local build directory (on PC)
 		/// </summary>
diff --git a/Engine/Source/Programs/IOS/iPhonePackager/FileOperations.cs b/Engine/Source/Programs/IOS/iPhonePackager/FileOperations.cs
--- a/Engine/Source/Programs/IOS/iPhonePackager/FileOperations.cs
+++ b/Engine/Source/Programs/IOS/iPhonePackager/FileOperations.cs
@@ -254,19 +254,24 @@ namespace iPhonePackager
 
 		static public bool BatchUploadFolder(string MacName, string SourceFolder, string DestFolder, bool bDeleteTarget)
 		{
-			if (!ConditionalInitRemoting(MacName))
+			if (Config.bUseRPCUtil)
 			{
-				return false;
-			}
 
-			Program.Log(" ... '" + SourceFolder + "' -> '" + DestFolder + "'");
+				if (!ConditionalInitRemoting(MacName))
+				{
+					return false;
+				}
 
-			if (bDeleteTarget)
-			{
-				// tell Mac to delete target if requested
-				RPCUtility.CommandHelper.RemoveDirectory(RPCSocket, DestFolder);
+
+				if (bDeleteTarget)
+				{
+					// tell Mac to delete target if requested
+					RPCUtility.CommandHelper.RemoveDirectory(RPCSocket, DestFolder);
+				}
 			}
 
+			Program.Log(" ... '" + SourceFolder + "' -> '" + DestFolder + "'");
+
 			DirectoryInfo SourceFolderInfo = new DirectoryInfo(SourceFolder);
 			if (!SourceFolderInfo.Exists)
 			{
@@ -278,20 +283,36 @@ namespace iPhonePackager
 			List<string> FilesToUpload = new List<string>();
 			RecursiveBatchUploadFolder(SourceFolderInfo, DestFolder, FilesToUpload);
 
-			// send them off!
-			RPCUtility.CommandHelper.RPCBatchUpload(RPCSocket, FilesToUpload.ToArray());
+			if (Config.bUseRPCUtil)
+			{
+				// send them off!
+				RPCUtility.CommandHelper.RPCBatchUpload(RPCSocket, FilesToUpload.ToArray());
+			}
+			else
+			{
+				SSHCommandHelper.BatchUpload(MacName, FilesToUpload.ToArray());
+			}
 
 			return true;
 		}
 
 		static public bool DownloadFile(string MacName, string SourceFilename, string DestFilename)
 		{
-			if (!ConditionalInitRemoting(MacName))
+			Hashtable Results = null;
+
+			if (Config.bUseRPCUtil)
 			{
-				return false;
-			}
+				if (!ConditionalInitRemoting(MacName))
+				{
+					return false;
+				}
 
-			Hashtable Results = RPCUtility.CommandHelper.RPCDownload(RPCSocket, SourceFilename, DestFilename);
+				Results = RPCUtility.CommandHelper.RPCDownload(RPCSocket, SourceFilename, DestFilename);
+			}
+			else
+			{
+				Results = SSHCommandHelper.DownloadFile(MacName, SourceFilename, DestFilename);
+			}
 
 			// success if exitcode is 0 or unspecified
 			if (Results["ExitCode"] != null)
diff --git a/Engine/Source/Programs/IOS/iPhonePackager/SSHCommandHelper.cs b/Engine/Source/Programs/IOS/iPhonePackager/SSHCommandHelper.cs
new file mode 100644
--- /dev/null
+++ b/Engine/Source/Programs/IOS/iPhonePackager/SSHCommandHelper.cs
@@ -0,0 +1,381 @@
+﻿/**
+ * Copyright 1998-2015 Epic Games, Inc. All Rights Reserved.
+ */
+
+using System;
+using System.Collections;
+using System.Collections.Generic;
+using System.Diagnostics;
+using System.IO;
+using System.Text;
+
+namespace iPhonePackager
+{
+	public partial class SSHCommandHelper
+	{
+		private static string SSHExePath = "";
+		private static string SSHAuthentication = "";
+		private static string SSHUser = "";
+		private static string RSyncExePath = "";
+
+		/** 
+		 * Parse the commandline for SSH arguments
+		 */
+		public static bool ParseSSHProperties(string[] Arguments)
+		{
+			Int32 ArgIndex = 0;
+			foreach( string NextArg in Arguments)
+			{
+				string Arg = NextArg.ToLowerInvariant();
+
+				if (Arg.StartsWith("-"))
+				{
+					if (Arg == "-usessh")
+					{
+						Config.bUseRPCUtil = false;
+					}
+					else if (Arg == "-sshexe")
+					{
+						if (Arguments.Length > ArgIndex + 1)
+						{
+							SSHExePath = Arguments[ArgIndex+1];
+						}
+						else
+						{
+							return false;
+						}
+					}
+					else if (Arg == "-sshauth")
+					{
+						if (Arguments.Length > ArgIndex + 1)
+						{
+							SSHAuthentication = Arguments[ArgIndex + 1];
+						}
+						else
+						{
+							return false;
+						}
+					}
+					else if (Arg == "-sshuser")
+					{
+						if (Arguments.Length > ArgIndex + 1)
+						{
+							SSHUser = Arguments[ArgIndex + 1];
+						}
+						else
+						{
+							return false;
+						}
+					}
+					else if (Arg == "-rsyncexe")
+					{
+						if (Arguments.Length > ArgIndex + 1)
+						{
+							RSyncExePath = Arguments[ArgIndex + 1];
+						}
+						else
+						{
+							return false;
+						}
+					}
+					else if (Arg == "-overridedevroot")
+					{
+						if (Arguments.Length > ArgIndex + 1)
+						{
+							Config.OverrideDevRoot = Arguments[ArgIndex + 1];
+						}
+						else
+						{
+							return false;
+						}
+					}
+				}
+				++ArgIndex;
+			}
+
+			bool bValidSetup = true;
+			if (Config.bUseRPCUtil == false)
+			{
+				// If we aren't using rpc utility, 
+				// ensure we have all the arguments needed for our ssh and rsync commands to run
+				bValidSetup = ValidateSetup();
+			}
+
+			return bValidSetup;
+		}
+
+		private static bool ValidateSetup()
+		{
+			bool bIsSetupForSSH = true;
+
+			if( string.IsNullOrEmpty(SSHExePath) )
+			{
+				Console.WriteLine("SSH requires the -sshexe commandline parameter");
+				bIsSetupForSSH = false;
+			}
+			else if (string.IsNullOrEmpty(SSHAuthentication))
+			{
+				Console.WriteLine("SSH requires the -sshauth commandline parameter");
+				bIsSetupForSSH = false;
+			}
+			else if (string.IsNullOrEmpty(SSHUser))
+			{
+				Console.WriteLine("SSH requires the -sshuser commandline parameter");
+				bIsSetupForSSH = false;
+			}
+			else if (string.IsNullOrEmpty(RSyncExePath))
+			{
+				Console.WriteLine("SSH requires the -rsyncexe commandline parameter");
+				bIsSetupForSSH = false;
+			}
+			else if (string.IsNullOrEmpty(Config.OverrideDevRoot))
+			{
+				Console.WriteLine("For now, SSH requires the -overridedevroot commandline parameter");
+				bIsSetupForSSH = false;
+			}
+
+			return bIsSetupForSSH;
+		}
+
+
+		// SSH command output, tracks output and errors from the remote execution
+		private static Dictionary<Object, StringBuilder> SSHOutputMap = new Dictionary<object, StringBuilder>();
+
+		/** 
+		 * Execute a console command on the remote mac.
+		 */
+		static public bool Command(string RemoteMac, string Command, string WorkingFolder)
+		{
+			// make the commandline for other end
+			string RemoteCommandline = "cd " + WorkingFolder + " && " + Command;
+
+			Process SSHProcess = new Process();
+
+			// Get the executable dir for SSH
+			string ExeDir = Path.GetDirectoryName(SSHExePath);
+			if (ExeDir != "")
+			{
+				SSHProcess.StartInfo.WorkingDirectory = ExeDir;
+			}
+
+			SSHProcess.StartInfo.FileName = SSHExePath;
+			SSHProcess.StartInfo.Arguments = string.Format(
+				"{0} {1}@{2} \"{3}\"",
+				SSHAuthentication,
+				SSHUser,
+				RemoteMac,
+				RemoteCommandline.Replace("\"", "\\\""));
+
+			// add this process to the map
+			SSHOutputMap[SSHProcess] = new StringBuilder("");
+
+			SSHProcess.StartInfo.UseShellExecute = false;
+			SSHProcess.StartInfo.RedirectStandardOutput = true;
+			SSHProcess.OutputDataReceived += new DataReceivedEventHandler(OutputReceivedForRsync);
+			SSHProcess.StartInfo.RedirectStandardError = true;
+			SSHProcess.ErrorDataReceived += new DataReceivedEventHandler(OutputReceivedForRsync);
+
+			DateTime Start = DateTime.Now;
+			SSHProcess.Start();
+
+			SSHProcess.BeginOutputReadLine();
+			SSHProcess.BeginErrorReadLine();
+
+			SSHProcess.WaitForExit();
+
+			Console.WriteLine("Execute took {0}", (DateTime.Now - Start).ToString());
+
+			// now we have enough to fill out the HashTable
+			Hashtable Return = new Hashtable();
+			Return["CommandOutput"] = SSHOutputMap[SSHProcess].ToString();
+			Return["ExitCode"] = (object)SSHProcess.ExitCode;
+
+			SSHOutputMap.Remove(SSHProcess);
+
+			return SSHProcess.ExitCode == 0;
+		}
+
+		/** 
+		 * Upload a number of files to the remote mac
+		 */
+		public static void BatchUpload(string RemoteMac, string[] CopyCommands)
+		{
+			// make a new structure to send over the wire with local filetime, and remote filename
+			List<string> LocalNames = new List<string>();
+			List<string> RemoteNames = new List<string>();
+
+			foreach (string Pair in CopyCommands)
+			{
+				// each pair is a local and remote filename
+				string[] Filenames = Pair.Split(";".ToCharArray());
+
+				if (Filenames.Length == 2)
+				{
+					FileInfo Info = new FileInfo(Filenames[0]);
+					if (Info.Exists)
+					{
+						// make the remote filename -> local filetime pair
+						LocalNames.Add(Filenames[0]);
+						RemoteNames.Add(Filenames[1]);
+					}
+					else
+					{
+						Console.WriteLine("Tried to batch upload non-existant file: " + Filenames[0]);
+					}
+				}
+			}
+
+			
+			for (int Index = 0; Index < LocalNames.Count; Index++)
+			{
+				// now send it over
+				Hashtable Result = UploadFile(RemoteMac, LocalNames[Index], RemoteNames[Index]);
+
+				if (Result == null)
+				{
+					// we want to fail here so builds don't quietly fail, but we use a useful message
+					throw new Exception(string.Format("Failed to upload local file {0} to {1}", LocalNames[Index], RemoteNames[Index]));
+				}
+			}
+		}
+
+		/** 
+		 * Upload a single file to the remote mac
+		 */
+		static public Hashtable UploadFile(string RemoteMac, string LocalPath, string RemotePath)
+		{
+			string RemoteDir = Path.GetDirectoryName(RemotePath).Replace("\\", "/");
+			string RemoteFilename = Path.GetFileName(RemotePath);
+
+			// get the executable dir for SSH, so Rsync can call it easily
+			string ExeDir = Path.GetDirectoryName(SSHExePath);
+
+			Process RsyncProcess = new Process();
+			if (ExeDir != "")
+			{
+				RsyncProcess.StartInfo.WorkingDirectory = ExeDir;
+			}
+
+			// make simple rsync commandline to send a file
+			RsyncProcess.StartInfo.FileName = RSyncExePath;
+			RsyncProcess.StartInfo.Arguments = string.Format(
+				"-zae \"ssh {0}\" --rsync-path=\"mkdir -p {1} && rsync\" \"{2}\" {3}@{4}:\"{1}/{5}\"",
+				SSHAuthentication,
+				RemoteDir,
+				ConvertPathToCygwin(LocalPath),
+				SSHUser,
+				RemoteMac,
+				RemoteFilename
+				);
+
+			SSHOutputMap[RsyncProcess] = new StringBuilder("");
+
+			RsyncProcess.StartInfo.UseShellExecute = false;
+			RsyncProcess.StartInfo.RedirectStandardOutput = true;
+			RsyncProcess.OutputDataReceived += new DataReceivedEventHandler(OutputReceivedForRsync);
+			RsyncProcess.StartInfo.RedirectStandardError = true;
+			RsyncProcess.ErrorDataReceived += new DataReceivedEventHandler(OutputReceivedForRsync);
+
+			RsyncProcess.Start();
+
+			RsyncProcess.BeginOutputReadLine();
+			RsyncProcess.BeginErrorReadLine();
+
+			RsyncProcess.WaitForExit();
+
+			// now we have enough to fill out the HashTable
+			Hashtable ReturnHash = new Hashtable();
+			ReturnHash["CommandOutput"] = SSHOutputMap[RsyncProcess].ToString();
+			ReturnHash["ExitCode"] = (object)RsyncProcess.ExitCode;
+
+			// run rsync (0 means success)
+			return ReturnHash;
+		}
+
+
+		/** 
+		 * Download a file from the remote mac
+		 */
+		static public Hashtable DownloadFile(string RemoteMac, string RemotePath, string LocalPath)
+		{
+			// get the executable dir for SSH, so Rsync can call it easily
+			string ExeDir = Path.GetDirectoryName(SSHExePath);
+
+			Process RsyncProcess = new Process();
+			if (ExeDir != "")
+			{
+				RsyncProcess.StartInfo.WorkingDirectory = ExeDir;
+			}
+
+			// make sure directory exists to download to
+			Directory.CreateDirectory(Path.GetDirectoryName(LocalPath));
+
+			// make simple rsync commandline to send a file
+			RsyncProcess.StartInfo.FileName = RSyncExePath;
+			RsyncProcess.StartInfo.Arguments = string.Format(
+				"-zae \"ssh {0}\" {2}@{3}:\"{4}\" \"{1}\"",
+				SSHAuthentication,
+				ConvertPathToCygwin(LocalPath),
+				SSHUser,
+				RemoteMac,
+				RemotePath
+				);
+
+			SSHOutputMap[RsyncProcess] = new StringBuilder("");
+
+			RsyncProcess.StartInfo.UseShellExecute = false;
+			RsyncProcess.StartInfo.RedirectStandardOutput = true;
+			RsyncProcess.OutputDataReceived += new DataReceivedEventHandler(OutputReceivedForRsync);
+			RsyncProcess.StartInfo.RedirectStandardError = true;
+			RsyncProcess.ErrorDataReceived += new DataReceivedEventHandler(OutputReceivedForRsync);
+
+			RsyncProcess.Start();
+
+			RsyncProcess.BeginOutputReadLine();
+			RsyncProcess.BeginErrorReadLine();
+
+			RsyncProcess.WaitForExit();
+
+			// now we have enough to fill out the HashTable
+			Hashtable ReturnHash = new Hashtable();
+			ReturnHash["CommandOutput"] = SSHOutputMap[RsyncProcess].ToString();
+			ReturnHash["ExitCode"] = (Int64)RsyncProcess.ExitCode;
+
+			// run rsync (0 means success)
+			return ReturnHash;
+		}
+
+
+		static public void OutputReceivedForRsync(Object Sender, DataReceivedEventArgs Line)
+		{
+			if ((Line != null) && (Line.Data != null) && (Line.Data != ""))
+			{
+				Program.Log(Line.Data);
+			}
+		}
+
+
+		static public void OutputReceivedForSSH(Object Sender, DataReceivedEventArgs Line)
+		{
+			if ((Line != null) && (Line.Data != null) && (Line.Data != ""))
+			{
+				StringBuilder SSHOutput = SSHOutputMap[Sender];
+				if (SSHOutput.Length != 0)
+				{
+					SSHOutput.Append(Environment.NewLine);
+				}
+				SSHOutput.Append(Line.Data);
+			}
+		}
+
+
+		private static string ConvertPathToCygwin(string InPath)
+		{
+			if (InPath == null)
+			{
+				return null;
+			}
+			return "/cygdrive/" + InPath.Replace("\\\\", "/").Replace("\\", "/").Replace(":", "");
+		}
+	}
+}
\ No newline at end of file
diff --git a/Engine/Source/Programs/IOS/iPhonePackager/iPhonePackager.cs b/Engine/Source/Programs/IOS/iPhonePackager/iPhonePackager.cs
--- a/Engine/Source/Programs/IOS/iPhonePackager/iPhonePackager.cs
+++ b/Engine/Source/Programs/IOS/iPhonePackager/iPhonePackager.cs
@@ -307,6 +307,11 @@ namespace iPhonePackager
 						MainRPCCommand = Arguments[ArgIndex];
 					}
 				}
+				
+				if(!SSHCommandHelper.ParseSSHProperties(Arguments))
+				{
+					return false;
+				}
 			}
 	
 			return ( true );
diff --git a/Engine/Source/Programs/IOS/iPhonePackager/iPhonePackager.csproj b/Engine/Source/Programs/IOS/iPhonePackager/iPhonePackager.csproj
--- a/Engine/Source/Programs/IOS/iPhonePackager/iPhonePackager.csproj
+++ b/Engine/Source/Programs/IOS/iPhonePackager/iPhonePackager.csproj
@@ -140,6 +140,7 @@
     <Compile Include="SlowProgressDialog.designer.cs">
       <DependentUpon>SlowProgressDialog.cs</DependentUpon>
     </Compile>
+    <Compile Include="SSHCommandHelper.cs" />
     <Compile Include="ToolsHub.cs">
       <SubType>Form</SubType>
     </Compile>
diff --git a/Engine/Source/Programs/UnrealBuildTool/IOS/IOSToolChain.cs b/Engine/Source/Programs/UnrealBuildTool/IOS/IOSToolChain.cs
--- a/Engine/Source/Programs/UnrealBuildTool/IOS/IOSToolChain.cs
+++ b/Engine/Source/Programs/UnrealBuildTool/IOS/IOSToolChain.cs
@@ -1299,18 +1299,8 @@ namespace UnrealBuildTool
 
 					// now that 
 					Process StubGenerateProcess = new Process();
-
-					if (RemoteToolChain.bUseRPCUtil)
-					{
-						StubGenerateProcess.StartInfo.WorkingDirectory = Path.GetFullPath("..\\Binaries\\DotNET\\IOS");
-						StubGenerateProcess.StartInfo.FileName = Path.Combine(StubGenerateProcess.StartInfo.WorkingDirectory, "iPhonePackager.exe");
-					}
-					else
-					{
-						// UAT wants the working directory to be the root (or does it matter?)
-						StubGenerateProcess.StartInfo.WorkingDirectory = Path.GetFullPath("..\\Binaries\\DotNET");
-						StubGenerateProcess.StartInfo.FileName = Path.Combine(StubGenerateProcess.StartInfo.WorkingDirectory, "AutomationTool.exe");
-					}
+					StubGenerateProcess.StartInfo.WorkingDirectory = Path.GetFullPath("..\\Binaries\\DotNET\\IOS");
+					StubGenerateProcess.StartInfo.FileName = Path.Combine(StubGenerateProcess.StartInfo.WorkingDirectory, "iPhonePackager.exe");
 
 					string Arguments = "";
 						
@@ -1358,19 +1348,20 @@ namespace UnrealBuildTool
 						Arguments += " -architecture " + Architecture.Substring(1);
 					}
 
-					// programmers that use Xcode packaging mode should use the following commandline instead, as it will package for Xcode on each compile
-					//				Arguments = "PackageApp " + GameName + " " + Configuration;
-
-					if (RemoteToolChain.bUseRPCUtil)
-					{
-						StubGenerateProcess.StartInfo.Arguments = Arguments;
-					}
-					else
+					if (!bUseRPCUtil )
 					{
-						Arguments = Arguments.Replace("\"", "\\\"");
-						StubGenerateProcess.StartInfo.Arguments = "IPhonePackager -cmd=\"" + Arguments + "\" -nocompile";
+						Arguments += " -usessh";
+						Arguments += " -sshexe \"" + ResolvedSSHExe + "\"";
+						Arguments += " -sshauth \"" + ResolvedSSHAuthentication + "\"";
+						Arguments += " -sshuser \"" + ResolvedRSyncUsername + "\"";
+						Arguments += " -rsyncexe \"" + ResolvedRSyncExe + "\"";
+						Arguments += " -overridedevroot \"" + UserDevRootMac + "\"";
 					}
 
+					// programmers that use Xcode packaging mode should use the following commandline instead, as it will package for Xcode on each compile
+					//				Arguments = "PackageApp " + GameName + " " + Configuration;
+					StubGenerateProcess.StartInfo.Arguments = Arguments;
+
 					StubGenerateProcess.OutputDataReceived += new DataReceivedEventHandler(OutputReceivedDataEventHandler);
 					StubGenerateProcess.ErrorDataReceived += new DataReceivedEventHandler(OutputReceivedDataEventHandler);
 
@@ -1448,13 +1439,11 @@ namespace UnrealBuildTool
 
 		public override string GetPlatformVersion()
 		{
-			ParseProjectSettings();
 			return RunTimeIOSVersion;
 		}
 
 		public override string GetPlatformDevices()
 		{
-			ParseProjectSettings();
 			return RunTimeIOSDevices;
 		}
 
diff --git a/Engine/Source/Programs/UnrealBuildTool/ToolChain/RemoteToolChain.cs b/Engine/Source/Programs/UnrealBuildTool/ToolChain/RemoteToolChain.cs
--- a/Engine/Source/Programs/UnrealBuildTool/ToolChain/RemoteToolChain.cs
+++ b/Engine/Source/Programs/UnrealBuildTool/ToolChain/RemoteToolChain.cs
@@ -763,7 +763,7 @@ namespace UnrealBuildTool
 			// make simple rsync commandline to send a file
 			RsyncProcess.StartInfo.FileName = ResolvedRSyncExe;
 			RsyncProcess.StartInfo.Arguments = string.Format(
-				"-zae {0} {2}@{3}:\"{4}\" \"{1}\"",
+				"-zae \"{0}\" {2}@{3}:\"{4}\" \"{1}\"",
 				ResolvedRsyncAuthentication,
 				ConvertPathToCygwin(LocalPath),
 				RSyncUsername,
