UEPLAT-23 Faster iteration time for content
Implemented iterative deploy framework and iOS specific implementation
When utilizing iterative deploy in conjunction with incremental cooking only the changed items are pushed to the device resulting in potentially far faster iteration times.
#ios

diff --git a/Engine/Source/Developer/LauncherServices/Private/Launcher/LauncherWorker.cpp b/Engine/Source/Developer/LauncherServices/Private/Launcher/LauncherWorker.cpp
--- a/Engine/Source/Developer/LauncherServices/Private/Launcher/LauncherWorker.cpp
+++ b/Engine/Source/Developer/LauncherServices/Private/Launcher/LauncherWorker.cpp
@@ -183,8 +183,7 @@ FString FLauncherWorker::CreateUATCommand( const ILauncherProfileRef& InProfile,
 {
 	CommandStart = TEXT("");
 	FString UATCommand = TEXT("");
-	FGuid InstanceId(FGuid::NewGuid());
-	FGuid SessionId(FGuid::NewGuid());
+	static FGuid SessionId(FGuid::NewGuid());
 	FString InitialMap = InProfile->GetDefaultLaunchRole()->GetInitialMap();
 	if (InitialMap.IsEmpty() && InProfile->GetCookedMaps().Num() == 1)
 	{
@@ -357,7 +356,7 @@ FString FLauncherWorker::CreateUATCommand( const ILauncherProfileRef& InProfile,
 
 			if (InProfile->IsCookingIncrementally())
 			{
-				UATCommand += TEXT(" -iterate");
+				UATCommand += TEXT(" -iterativecooking");
 			}
 
 			if (InProfile->IsCookingUnversioned())
@@ -413,6 +412,10 @@ FString FLauncherWorker::CreateUATCommand( const ILauncherProfileRef& InProfile,
 		break;
 	case ELauncherProfileCookModes::ByTheBookInEditor:
 		UATCommand += MapList;
+		if (InProfile->IsCookingIncrementally())
+		{
+			UATCommand += TEXT(" -iterativecooking");
+		}
 	case ELauncherProfileCookModes::DoNotCook:
 		UATCommand += TEXT(" -skipcook");
 		break;
@@ -448,6 +451,10 @@ FString FLauncherWorker::CreateUATCommand( const ILauncherProfileRef& InProfile,
 		case ELauncherProfileDeploymentModes::CopyToDevice:
 			{
 				UATCommand += TEXT(" -stage -deploy");
+				if (Profile->IsDeployingIncrementally())
+				{
+					UATCommand += " -iterativedeploy";
+				}
 				UATCommand += CommandLine;
 				UATCommand += StageDirectory;
 				UATCommand += DeviceCommand;
@@ -470,7 +477,7 @@ FString FLauncherWorker::CreateUATCommand( const ILauncherProfileRef& InProfile,
 		// run
 		if (InProfile->GetLaunchMode() != ELauncherProfileLaunchModes::DoNotLaunch)
 		{
-			UATCommand += TEXT(" -run");
+			UATCommand += TEXT(" -run -nokill");
 
 			FCommandDesc Desc;
 			FText Command = FText::Format(LOCTEXT("LauncherRunDesc", "Launching on {0}"), FText::FromString(DeviceNames.RightChop(1)));
diff --git a/Engine/Source/Developer/LauncherServices/Private/LauncherServicesPrivatePCH.h b/Engine/Source/Developer/LauncherServices/Private/LauncherServicesPrivatePCH.h
--- a/Engine/Source/Developer/LauncherServices/Private/LauncherServicesPrivatePCH.h
+++ b/Engine/Source/Developer/LauncherServices/Private/LauncherServicesPrivatePCH.h
@@ -24,7 +24,8 @@
 #define LAUNCHERSERVICES_SIMPLEPROFILEVERSION 1
 
 /** Defines the launcher profile file version. */
-#define LAUNCHERSERVICES_PROFILEVERSION 10
+#define LAUNCHERSERVICES_PROFILEVERSION 11
+#define LAUNCHERSERVICES_MINPROFILEVERSION 10
 
 
 // profile manager
diff --git a/Engine/Source/Developer/LauncherServices/Private/Profiles/LauncherProfile.h b/Engine/Source/Developer/LauncherServices/Private/Profiles/LauncherProfile.h
--- a/Engine/Source/Developer/LauncherServices/Private/Profiles/LauncherProfile.h
+++ b/Engine/Source/Developer/LauncherServices/Private/Profiles/LauncherProfile.h
@@ -2,6 +2,8 @@
 
 #pragma once
 
+#define LAUNCHERSERVICES_ADDEDINCREMENTALDEPLOYVERSION 11
+
 /**
 * Implements a simple profile which controls the desired output of the Launcher for simple
 */
@@ -423,6 +425,11 @@ public:
 		return true;
 	}
 
+	virtual bool IsDeployingIncrementally( ) const override
+	{
+		return DeployIncremental;
+	}
+
 	virtual bool IsFileServerHidden( ) const override
 	{
 		return HideFileServerWindow;
@@ -477,7 +484,7 @@ public:
 
 		Archive	<< Version;
 
-		if (Version != LAUNCHERSERVICES_PROFILEVERSION)
+		if (Version < LAUNCHERSERVICES_MINPROFILEVERSION)
 		{
 			return false;
 		}
@@ -521,6 +528,11 @@ public:
                 << ForceClose
                 << Timeout;
 
+		if (Version >= LAUNCHERSERVICES_ADDEDINCREMENTALDEPLOYVERSION)
+		{
+			Archive << DeployIncremental;
+		}
+
 		DefaultLaunchRole->Serialize(Archive);
 
 		// serialize launch roles
@@ -607,6 +619,7 @@ public:
 		DeployWithUnrealPak = false;
 		DeployedDeviceGroupId = FGuid();
 		HideFileServerWindow = false;
+		DeployIncremental = false;
 
 		// default launch settings
 		LaunchMode = ELauncherProfileLaunchModes::DefaultRole;
@@ -763,6 +776,16 @@ public:
 		}
 	}
 
+	virtual void SetIncrementalDeploying( bool Incremental ) override
+	{
+		if (DeployIncremental != Incremental)
+		{
+			DeployIncremental = Incremental;
+
+			Validate();
+		}
+	}
+
 	virtual void SetLaunchMode( ELauncherProfileLaunchModes::Type Mode ) override
 	{
 		if (LaunchMode != Mode)
@@ -1117,6 +1140,9 @@ private:
 	// Holds a flag indicating whether content should be packaged with UnrealPak.
 	bool DeployWithUnrealPak;
 
+	// Holds a flag indicating whether to use incremental deployment
+	bool DeployIncremental;
+
 	// Holds the device group to deploy to.
 	ILauncherDeviceGroupPtr DeployedDeviceGroup;
 
diff --git a/Engine/Source/Developer/LauncherServices/Public/Interfaces/ILauncherProfile.h b/Engine/Source/Developer/LauncherServices/Public/Interfaces/ILauncherProfile.h
--- a/Engine/Source/Developer/LauncherServices/Public/Interfaces/ILauncherProfile.h
+++ b/Engine/Source/Developer/LauncherServices/Public/Interfaces/ILauncherProfile.h
@@ -580,6 +580,14 @@ public:
 	 */
 	virtual bool IsCookingUnversioned( ) const = 0;
 
+	/**
+	 * Checks whether incremental deployment is enabled.
+	 *
+	 * @return true if deploying incrementally, false otherwise.
+	 * @see SetIncrementalDeploying
+	 */
+	virtual bool IsDeployingIncrementally( ) const = 0;
+
 	/**
 	 * Checks whether the file server's console window should be hidden.
 	 *
@@ -811,6 +819,14 @@ public:
 	 */
 	virtual void SetIncrementalCooking( bool Incremental ) = 0;
 
+	/**
+	 * Sets incremental deploying.
+	 *
+	 * @param Incremental Whether deploying should be incremental.
+	 * @see IsDeployingIncrementally
+	 */
+	virtual void SetIncrementalDeploying( bool Incremental ) = 0;
+
 	/**
 	 * Sets the launch mode.
 	 *
diff --git a/Engine/Source/Developer/ProjectLauncher/Private/Widgets/Deploy/SProjectLauncherDeployToDeviceSettings.cpp b/Engine/Source/Developer/ProjectLauncher/Private/Widgets/Deploy/SProjectLauncherDeployToDeviceSettings.cpp
--- a/Engine/Source/Developer/ProjectLauncher/Private/Widgets/Deploy/SProjectLauncherDeployToDeviceSettings.cpp
+++ b/Engine/Source/Developer/ProjectLauncher/Private/Widgets/Deploy/SProjectLauncherDeployToDeviceSettings.cpp
@@ -30,7 +30,7 @@ void SProjectLauncherDeployToDeviceSettings::Construct( const FArguments& InArgs
 					]
 			]
 
-/*		+ SVerticalBox::Slot()
+		+ SVerticalBox::Slot()
 			.AutoHeight()
 			.Padding(0.0f, 8.0f, 0.0f, 0.0f)
 			[
@@ -46,9 +46,26 @@ void SProjectLauncherDeployToDeviceSettings::Construct( const FArguments& InArgs
 							.Padding(8.0f)
 							.BodyContent()
 							[
+								SNew(SVerticalBox)
+
+								+ SVerticalBox::Slot()
+								.AutoHeight()
+								[
+									// incremental cook check box
+									SNew(SCheckBox)
+									.IsChecked(this, &SProjectLauncherDeployToDeviceSettings::HandleIncrementalCheckBoxIsChecked)
+									.OnCheckStateChanged(this, &SProjectLauncherDeployToDeviceSettings::HandleIncrementalCheckBoxCheckStateChanged)
+									.Padding(FMargin(4.0f, 0.0f))
+									.ToolTipText(LOCTEXT("IncrementalCheckBoxTooltip", "If checked, only modified content will be deployed, resulting in much faster deploy times. It is recommended to enable this option whenever possible."))
+									.Content()
+									[
+										SNew(STextBlock)
+										.Text(LOCTEXT("IncrementalCheckBoxText", "Only deploy modified content"))
+									]
+								]
 							]
 					]
-			]*/
+			]
 	];
 }
 
@@ -57,4 +74,30 @@ void SProjectLauncherDeployToDeviceSettings::Construct( const FArguments& InArgs
  *****************************************************************************/
 
 
+void SProjectLauncherDeployToDeviceSettings::HandleIncrementalCheckBoxCheckStateChanged( ECheckBoxState NewState )
+{
+	ILauncherProfilePtr SelectedProfile = Model->GetSelectedProfile();
+
+	if (SelectedProfile.IsValid())
+	{
+		SelectedProfile->SetIncrementalDeploying(NewState == ECheckBoxState::Checked);
+	}
+}
+
+
+ECheckBoxState SProjectLauncherDeployToDeviceSettings::HandleIncrementalCheckBoxIsChecked( ) const
+{
+	ILauncherProfilePtr SelectedProfile = Model->GetSelectedProfile();
+
+	if (SelectedProfile.IsValid())
+	{
+		if (SelectedProfile->IsDeployingIncrementally())
+		{
+			return ECheckBoxState::Checked;
+		}
+	}
+
+	return ECheckBoxState::Unchecked;
+}
+
 #undef LOCTEXT_NAMESPACE
diff --git a/Engine/Source/Developer/ProjectLauncher/Private/Widgets/Deploy/SProjectLauncherDeployToDeviceSettings.h b/Engine/Source/Developer/ProjectLauncher/Private/Widgets/Deploy/SProjectLauncherDeployToDeviceSettings.h
--- a/Engine/Source/Developer/ProjectLauncher/Private/Widgets/Deploy/SProjectLauncherDeployToDeviceSettings.h
+++ b/Engine/Source/Developer/ProjectLauncher/Private/Widgets/Deploy/SProjectLauncherDeployToDeviceSettings.h
@@ -24,10 +24,16 @@ public:
 	 */
 	void Construct(	const FArguments& InArgs, const FProjectLauncherModelRef& InModel, EVisibility InShowAdvanced = EVisibility::Visible );
 
-private:
+protected:
 
 private:
 
+	/** Handles check state changes of the 'Incremental' check box. */
+	void HandleIncrementalCheckBoxCheckStateChanged( ECheckBoxState NewState );
+
+	/** Handles determining the checked state of the 'Incremental' check box. */
+	ECheckBoxState HandleIncrementalCheckBoxIsChecked( ) const;
+
 	// Holds a pointer to the data model.
 	FProjectLauncherModelPtr Model;
 };
diff --git a/Engine/Source/Editor/UnrealEd/Private/PlayLevel.cpp b/Engine/Source/Editor/UnrealEd/Private/PlayLevel.cpp
--- a/Engine/Source/Editor/UnrealEd/Private/PlayLevel.cpp
+++ b/Engine/Source/Editor/UnrealEd/Private/PlayLevel.cpp
@@ -1210,7 +1210,7 @@ void UEditorEngine::PlayStandaloneLocalPc(FString MapNameOverride, FIntPoint* Wi
 
 static void HandleOutputReceived(const FString& InMessage)
 {
-	UE_LOG(LogPlayLevel, Display, TEXT("%s"), *InMessage);
+	UE_LOG(LogPlayLevel, Log, TEXT("%s"), *InMessage);
 }
 
 static void HandleCancelButtonClicked(ILauncherWorkerPtr LauncherWorker)
@@ -1403,8 +1403,10 @@ void UEditorEngine::PlayUsingLauncher()
 		ILauncherProfileRef LauncherProfile = LauncherServicesModule.CreateProfile(TEXT("Play On Device"));
 		LauncherProfile->SetBuildGame(bHasCode && FSourceCodeNavigation::IsCompilerAvailable());
 		LauncherProfile->SetCookMode(CanCookByTheBookInEditor() ? ELauncherProfileCookModes::ByTheBookInEditor : ELauncherProfileCookModes::ByTheBook);
+		LauncherProfile->SetIncrementalCooking(true);
 		LauncherProfile->AddCookedPlatform(PlayUsingLauncherDeviceId.Left(PlayUsingLauncherDeviceId.Find(TEXT("@"))));
 		LauncherProfile->SetDeployedDeviceGroup(DeviceGroup);
+		LauncherProfile->SetIncrementalDeploying(true);
 		LauncherProfile->SetEditorExe(FUnrealEdMisc::Get().GetExecutableForCommandlets());
 
 		const FString DummyDeviceName(FString::Printf(TEXT("All_iOS_On_%s"), FPlatformProcess::ComputerName()));
diff --git a/Engine/Source/Programs/AutomationTool/IOS/IOSPlatform.Automation.cs b/Engine/Source/Programs/AutomationTool/IOS/IOSPlatform.Automation.cs
--- a/Engine/Source/Programs/AutomationTool/IOS/IOSPlatform.Automation.cs
+++ b/Engine/Source/Programs/AutomationTool/IOS/IOSPlatform.Automation.cs
@@ -18,6 +18,8 @@ static class IOSEnvVarNames
 
 public class IOSPlatform : Platform
 {
+	bool bCreatedIPA = false;
+
 	public IOSPlatform()
 		:base(UnrealTargetPlatform.IOS)
 	{
@@ -152,27 +154,52 @@ public class IOSPlatform : Platform
 			InternalUtils.SafeDeleteDirectory(AppDirectory + "/cookeddata", true);
 			InternalUtils.SafeDeleteFile(AppDirectory + "/ue4commandline.txt", true);
 
-			// copy the Staged files to the AppDirectory
-			string[] StagedFiles = Directory.GetFiles(SC.StageDirectory, "*", SearchOption.AllDirectories);
-			foreach (string Filename in StagedFiles)
+			if (!Params.IterativeDeploy)
 			{
-				string DestFilename = Filename.Replace(SC.StageDirectory, AppDirectory);
-				Directory.CreateDirectory(Path.GetDirectoryName(DestFilename));
-				InternalUtils.SafeCopyFile(Filename, DestFilename, true);
+				// copy the Staged files to the AppDirectory
+				string[] StagedFiles = Directory.GetFiles (SC.StageDirectory, "*", SearchOption.AllDirectories);
+				foreach (string Filename in StagedFiles)
+				{
+					string DestFilename = Filename.Replace (SC.StageDirectory, AppDirectory);
+					Directory.CreateDirectory (Path.GetDirectoryName (DestFilename));
+					InternalUtils.SafeCopyFile (Filename, DestFilename, true);
+				}
+			}
+			else
+			{
+				// copy just the root stage directory files
+				string[] StagedFiles = Directory.GetFiles (SC.StageDirectory, "*", SearchOption.TopDirectoryOnly);
+				foreach (string Filename in StagedFiles)
+				{
+					string DestFilename = Filename.Replace (SC.StageDirectory, AppDirectory);
+					Directory.CreateDirectory (Path.GetDirectoryName (DestFilename));
+					InternalUtils.SafeCopyFile (Filename, DestFilename, true);
+				}
 			}
 		}
 
-
-		if (UnrealBuildTool.BuildHostPlatform.Current.Platform != UnrealTargetPlatform.Mac)
+		if (SC.StageTargetConfigurations.Count != 1)
+		{
+			throw new AutomationException("iOS is currently only able to package one target configuration at a time, but StageTargetConfigurations contained {0} configurations", SC.StageTargetConfigurations.Count);
+		}
+		bCreatedIPA = false;
+		bool bNeedsIPA = false;
+		if (Params.IterativeDeploy)
 		{
-			if (SC.StageTargetConfigurations.Count != 1) 
+			// check to determine if we need to update the IPA
+			if (File.Exists(SC.StageDirectory + "/Manifest_DeltaNonUFSFiles.txt"))
 			{
-				throw new AutomationException ("iOS is currently only able to package one target configuration at a time, but StageTargetConfigurations contained {0} configurations", SC.StageTargetConfigurations.Count);
+				string NonUFSFiles = File.ReadAllText(SC.StageDirectory + "/Manifest_DeltaNonUFSFiles.txt");
+				string[] Lines = NonUFSFiles.Split('\n');
+				bNeedsIPA = Lines.Length > 0 && !string.IsNullOrWhiteSpace(Lines[0]);
 			}
-			var TargetConfiguration = SC.StageTargetConfigurations[0];
+		}
 
-			var ProjectStub = Path.GetFullPath (Params.ProjectGameExeFilename) + "/" + Params.ShortProjectName + Path.GetExtension(Params.ProjectGameExeFilename);
+		if (UnrealBuildTool.BuildHostPlatform.Current.Platform != UnrealTargetPlatform.Mac)
+		{
+			var TargetConfiguration = SC.StageTargetConfigurations[0];
 			var ProjectIPA = MakeIPAFileName(TargetConfiguration, Params);
+			var ProjectStub = Path.GetFullPath(Params.ProjectGameExeFilename) + "/" + Params.ShortProjectName + Path.GetExtension(Params.ProjectGameExeFilename);
 
 			// package a .ipa from the now staged directory
 			var IPPExe = CombinePaths(CmdEnv.LocalRoot, "Engine/Binaries/DotNET/IOS/IPhonePackager.exe");
@@ -184,85 +211,99 @@ public class IOSPlatform : Platform
 
 			bool cookonthefly = Params.CookOnTheFly || Params.SkipCookOnTheFly;
 
-			// delete the .ipa to make sure it was made
-			DeleteFile(ProjectIPA);
-
-			if (RemoteToolChain.bUseRPCUtil)
+			// if we are incremental check to see if we need to even update the IPA
+			if (!Params.IterativeDeploy || !File.Exists(ProjectIPA) || bNeedsIPA)
 			{
-				string IPPArguments = "RepackageFromStage \"" + (Params.IsCodeBasedProject ? Params.RawProjectPath : "Engine") + "\"";
-				IPPArguments += " -config " + TargetConfiguration.ToString();
+				// delete the .ipa to make sure it was made
+				DeleteFile(ProjectIPA);
 
-				if (TargetConfiguration == UnrealTargetConfiguration.Shipping)
+				bCreatedIPA = true;
+
+				if (RemoteToolChain.bUseRPCUtil)
 				{
-					IPPArguments += " -compress=best";
-				}
+					string IPPArguments = "RepackageFromStage \"" + (Params.IsCodeBasedProject ? Params.RawProjectPath : "Engine") + "\"";
+					IPPArguments += " -config " + TargetConfiguration.ToString();
 
-				// Determine if we should sign
-				bool bNeedToSign = GetCodeSignDesirability(Params);
+					if (TargetConfiguration == UnrealTargetConfiguration.Shipping)
+					{
+						IPPArguments += " -compress=best";
+					}
 
-				if (!String.IsNullOrEmpty(Params.BundleName))
-				{
-					// Have to sign when a bundle name is specified
-					bNeedToSign = true;
-					IPPArguments += " -bundlename " + Params.BundleName;
-				}
+					// Determine if we should sign
+					bool bNeedToSign = GetCodeSignDesirability(Params);
 
-				if (bNeedToSign)
-				{
-					IPPArguments += " -sign";
-					if (Params.Distribution)
+					if (!String.IsNullOrEmpty(Params.BundleName))
 					{
-						IPPArguments += " -distribution";
+						// Have to sign when a bundle name is specified
+						bNeedToSign = true;
+						IPPArguments += " -bundlename " + Params.BundleName;
 					}
-				}
 
-				IPPArguments += (cookonthefly ? " -cookonthefly" : "");
-				IPPArguments += " -stagedir \"" + CombinePaths(Params.BaseStageDirectory, "IOS") + "\"";
-				IPPArguments += " -project \"" + Params.RawProjectPath + "\"";
+					if (bNeedToSign)
+					{
+						IPPArguments += " -sign";
+						if (Params.Distribution)
+						{
+							IPPArguments += " -distribution";
+						}
+					}
 
-				RunAndLog(CmdEnv, IPPExe, IPPArguments);
-			}
-			else
-			{
-				List<string> IPPArguments = new List<string>();
-				IPPArguments.Add("RepackageFromStage");
-				IPPArguments.Add(Params.IsCodeBasedProject ? Params.RawProjectPath : "Engine");
-				IPPArguments.Add("-config");
-				IPPArguments.Add(TargetConfiguration.ToString());
+					IPPArguments += (cookonthefly ? " -cookonthefly" : "");
+					IPPArguments += " -stagedir \"" + CombinePaths(Params.BaseStageDirectory, "IOS") + "\"";
+					IPPArguments += " -project \"" + Params.RawProjectPath + "\"";
+					if (Params.IterativeDeploy)
+					{
+						IPPArguments += " -iterate";
+					}
 
-				if (TargetConfiguration == UnrealTargetConfiguration.Shipping)
-				{
-					IPPArguments.Add("-compress=best");
+					RunAndLog(CmdEnv, IPPExe, IPPArguments);
 				}
+				else
+				{
+					List<string> IPPArguments = new List<string>();
+					IPPArguments.Add("RepackageFromStage");
+					IPPArguments.Add(Params.IsCodeBasedProject ? Params.RawProjectPath : "Engine");
+					IPPArguments.Add("-config");
+					IPPArguments.Add(TargetConfiguration.ToString());
 
-				// Determine if we should sign
-				bool bNeedToSign = GetCodeSignDesirability(Params);
+					if (TargetConfiguration == UnrealTargetConfiguration.Shipping)
+					{
+						IPPArguments.Add("-compress=best");
+					}
 
-				if (!String.IsNullOrEmpty(Params.BundleName))
-				{
-					// Have to sign when a bundle name is specified
-					bNeedToSign = true;
-					IPPArguments.Add("-bundlename");
-					IPPArguments.Add(Params.BundleName);
-				}
+					// Determine if we should sign
+					bool bNeedToSign = GetCodeSignDesirability(Params);
 
-				if (bNeedToSign)
-				{
-					IPPArguments.Add("-sign");
-				}
+					if (!String.IsNullOrEmpty(Params.BundleName))
+					{
+						// Have to sign when a bundle name is specified
+						bNeedToSign = true;
+						IPPArguments.Add("-bundlename");
+						IPPArguments.Add(Params.BundleName);
+					}
 
-				if (cookonthefly)
-				{
-					IPPArguments.Add(" -cookonthefly");
-				}
-				IPPArguments.Add(" -stagedir");
-				IPPArguments.Add(CombinePaths(Params.BaseStageDirectory, "IOS"));
-				IPPArguments.Add(" -project");
-				IPPArguments.Add(Params.RawProjectPath);
+					if (bNeedToSign)
+					{
+						IPPArguments.Add("-sign");
+					}
 
-				if (RunIPP(IPPArguments.ToArray()) != 0)
-				{
-					throw new AutomationException("IPP Failed");
+					if (cookonthefly)
+					{
+						IPPArguments.Add(" -cookonthefly");
+					}
+					IPPArguments.Add(" -stagedir");
+					IPPArguments.Add(CombinePaths(Params.BaseStageDirectory, "IOS"));
+					IPPArguments.Add(" -project");
+					IPPArguments.Add(Params.RawProjectPath);
+					if (Params.IterativeDeploy)
+					{
+						IPPArguments.Add(" -iterate");
+					}
+
+					if (RunIPP(IPPArguments.ToArray()) != 0)
+					{
+						throw new AutomationException("IPP Failed");
+					}
 				}
 			}
 
@@ -290,11 +331,19 @@ public class IOSPlatform : Platform
 		}
 		else
 		{
-			// code sign the app
-			CodeSign(Path.GetDirectoryName(Params.ProjectGameExeFilename), Params.IsCodeBasedProject ? Params.ShortProjectName : Path.GetFileNameWithoutExtension(Params.ProjectGameExeFilename), Params.RawProjectPath, SC.StageTargetConfigurations[0], SC.LocalRoot, Params.ShortProjectName, Path.GetDirectoryName(Params.RawProjectPath), SC.IsCodeBasedProject, Params.Distribution);
+			// create the ipa
+			string IPAName = CombinePaths(Path.GetDirectoryName(Params.RawProjectPath), "Binaries", "IOS", (Params.Distribution ? "Distro_" : "") + Params.ShortProjectName + (SC.StageTargetConfigurations[0] != UnrealTargetConfiguration.Development ? ("-IOS-" + SC.StageTargetConfigurations[0].ToString()) : "") + ".ipa");
 
-			// now generate the ipa
-			PackageIPA(Path.GetDirectoryName(Params.ProjectGameExeFilename), Params.IsCodeBasedProject ? Params.ShortProjectName : Path.GetFileNameWithoutExtension(Params.ProjectGameExeFilename), Params.ShortProjectName, Path.GetDirectoryName(Params.RawProjectPath), SC.StageTargetConfigurations[0], Params.Distribution);
+			if (!Params.IterativeDeploy || !File.Exists(IPAName) || bNeedsIPA)
+			{
+				bCreatedIPA = true;
+
+				// code sign the app
+				CodeSign(Path.GetDirectoryName(Params.ProjectGameExeFilename), Params.IsCodeBasedProject ? Params.ShortProjectName : Path.GetFileNameWithoutExtension(Params.ProjectGameExeFilename), Params.RawProjectPath, SC.StageTargetConfigurations[0], SC.LocalRoot, Params.ShortProjectName, Path.GetDirectoryName(Params.RawProjectPath), SC.IsCodeBasedProject, Params.Distribution);
+
+				// now generate the ipa
+				PackageIPA(Path.GetDirectoryName(Params.ProjectGameExeFilename), Params.IsCodeBasedProject ? Params.ShortProjectName : Path.GetFileNameWithoutExtension(Params.ProjectGameExeFilename), Params.ShortProjectName, Path.GetDirectoryName(Params.RawProjectPath), SC.StageTargetConfigurations[0], Params.Distribution);
+			}
 		}
 
 		PrintRunTime();
@@ -667,7 +716,7 @@ public class IOSPlatform : Platform
 
 	public override void GetFilesToDeployOrStage(ProjectParams Params, DeploymentContext SC)
 	{
-		if (UnrealBuildTool.BuildHostPlatform.Current.Platform != UnrealTargetPlatform.Mac)
+		//		if (UnrealBuildTool.BuildHostPlatform.Current.Platform != UnrealTargetPlatform.Mac)
 		{
 			// copy the icons/launch screens from the engine
 			{
@@ -750,41 +799,112 @@ public class IOSPlatform : Platform
 		SC.ArchiveFiles(Path.GetDirectoryName(ProjectIPA), Path.GetFileName(ProjectIPA));
 	}
 
+	public override bool RetrieveDeployedManifests(ProjectParams Params, DeploymentContext SC, out List<string> UFSManifests, out List<string> NonUFSManifests)
+	{
+		bool Result = true;
+		UFSManifests = new List<string>();
+		NonUFSManifests = new List<string>();
+		var DeployServer = CombinePaths(CmdEnv.LocalRoot, "Engine/Binaries/DotNET/IOS/DeploymentServer.exe");
+		try
+		{
+			var TargetConfiguration = SC.StageTargetConfigurations[0];
+			string BundleIdentifier = "";
+			if (File.Exists(Params.BaseStageDirectory + "/IOS/Info.plist"))
+			{
+				string Contents = File.ReadAllText(SC.StageDirectory + "/Info.plist");
+				int Pos = Contents.IndexOf("CFBundleIdentifier");
+				Pos = Contents.IndexOf("<string>", Pos) + 8;
+				int EndPos = Contents.IndexOf("</string>", Pos);
+				BundleIdentifier = Contents.Substring(Pos, EndPos - Pos);
+			}
+			RunAndLog(CmdEnv, DeployServer, "Backup -file \"" + CombinePaths(Params.BaseStageDirectory, "IOS", "Manifest_UFSFiles.txt") + "\" -file \"" + CombinePaths(Params.BaseStageDirectory, "IOS", "Manifest_NonUFSFiles.txt") + "\"" + (String.IsNullOrEmpty(Params.Device) ? "" : " -device " + Params.Device.Substring(4)) + " -bundle " + BundleIdentifier);
+
+			string[] ManifestFiles = Directory.GetFiles(CombinePaths(Params.BaseStageDirectory, "IOS"), "*_Manifest_UFS*.txt");
+			UFSManifests.AddRange(ManifestFiles);
+
+			ManifestFiles = Directory.GetFiles(CombinePaths(Params.BaseStageDirectory, "IOS"), "*_Manifest_NonUFS*.txt");
+			NonUFSManifests.AddRange(ManifestFiles);
+		}
+		catch (System.Exception)
+		{
+			// delete any files that did get copied
+			string[] Manifests = Directory.GetFiles(CombinePaths(Params.BaseStageDirectory, "IOS"), "*_Manifest_*.txt");
+			foreach (string Manifest in Manifests)
+			{
+				File.Delete(Manifest);
+			}
+			Result = false;
+		}
+
+		return Result;
+	}
+
 	public override void Deploy(ProjectParams Params, DeploymentContext SC)
     {
-		if (UnrealBuildTool.BuildHostPlatform.Current.Platform != UnrealTargetPlatform.Mac)
+		if (SC.StageTargetConfigurations.Count != 1)
 		{
-			if (SC.StageTargetConfigurations.Count != 1)
+			throw new AutomationException ("iOS is currently only able to package one target configuration at a time, but StageTargetConfigurations contained {0} configurations", SC.StageTargetConfigurations.Count);
+		}
+		if (Params.Distribution)
+		{
+			throw new AutomationException("iOS cannot deploy a package made for distribution.");
+		}
+		var TargetConfiguration = SC.StageTargetConfigurations[0];
+		var ProjectIPA = MakeIPAFileName(TargetConfiguration, Params);
+		var StagedIPA = SC.StageDirectory + "\\" + Path.GetFileName(ProjectIPA);
+
+		// verify the .ipa exists
+		if (!FileExists(StagedIPA))
+		{
+			StagedIPA = ProjectIPA;
+			if(!FileExists(StagedIPA))
 			{
-				throw new AutomationException ("iOS is currently only able to package one target configuration at a time, but StageTargetConfigurations contained {0} configurations", SC.StageTargetConfigurations.Count);
+				throw new AutomationException("DEPLOY FAILED - {0} was not found", ProjectIPA);
 			}
-			if (Params.Distribution)
+		}
+
+		// if iterative deploy, determine the file delta
+		string BundleIdentifier = "";
+		bool bNeedsIPA = true;
+		if (Params.IterativeDeploy)
+		{
+			if (File.Exists(Params.BaseStageDirectory + "/IOS/Info.plist"))
 			{
-				throw new AutomationException("iOS cannot deploy a package made for distribution.");
+				string Contents = File.ReadAllText(SC.StageDirectory + "/Info.plist");
+				int Pos = Contents.IndexOf("CFBundleIdentifier");
+				Pos = Contents.IndexOf("<string>", Pos) + 8;
+				int EndPos = Contents.IndexOf("</string>", Pos);
+				BundleIdentifier = Contents.Substring(Pos, EndPos - Pos);
 			}
-			var TargetConfiguration = SC.StageTargetConfigurations[0];
-			var ProjectIPA = MakeIPAFileName(TargetConfiguration, Params);
-			var StagedIPA = SC.StageDirectory + "\\" + Path.GetFileName(ProjectIPA);
 
-			// verify the .ipa exists
-			if (!FileExists(StagedIPA))
+			// check to determine if we need to update the IPA
+			if (File.Exists(SC.StageDirectory + "/Manifest_DeltaNonUFSFiles.txt"))
 			{
-				StagedIPA = ProjectIPA;
-				if(!FileExists(StagedIPA))
-				{
-					throw new AutomationException("DEPLOY FAILED - {0} was not found", ProjectIPA);
-				}
+				string NonUFSFiles = File.ReadAllText(SC.StageDirectory + "/Manifest_DeltaNonUFSFiles.txt");
+				string[] Lines = NonUFSFiles.Split('\n');
+				bNeedsIPA = Lines.Length > 0 && !string.IsNullOrWhiteSpace(Lines[0]);
 			}
+		}
 
-			// Add a commandline for this deploy, if the config allows it.
-			string AdditionalCommandline = (Params.FileServer || Params.CookOnTheFly) ? "" : (" -additionalcommandline " + "\"" + Params.RunCommandline + "\"");
+		// Add a commandline for this deploy, if the config allows it.
+		string AdditionalCommandline = (Params.FileServer || Params.CookOnTheFly) ? "" : (" -additionalcommandline " + "\"" + Params.RunCommandline + "\"");
 
-			// deploy the .ipa
-			var IPPExe = CombinePaths(CmdEnv.LocalRoot, "Engine/Binaries/DotNET/IOS/IPhonePackager.exe");
+		// deploy the .ipa
+		var DeployServer = CombinePaths(CmdEnv.LocalRoot, "Engine/Binaries/DotNET/IOS/DeploymentServer.exe");
 
-			// check for it in the stage directory
-			RunAndLog(CmdEnv, IPPExe, "Deploy \"" + Path.GetFullPath(StagedIPA) + "\"" + (String.IsNullOrEmpty(Params.Device) ? "" : " -device " + Params.Device.Substring(4)) + AdditionalCommandline);
+		// check for it in the stage directory
+		string CurrentDir = Directory.GetCurrentDirectory();
+		Directory.SetCurrentDirectory(CombinePaths(CmdEnv.LocalRoot, "Engine/Binaries/DotNET/IOS/"));
+		if (!Params.IterativeDeploy || bCreatedIPA || bNeedsIPA)
+		{
+			RunAndLog(CmdEnv, DeployServer, "Install -ipa \"" + Path.GetFullPath(StagedIPA) + "\"" + (String.IsNullOrEmpty(Params.Device) ? "" : " -device " + Params.Device.Substring(4)) + AdditionalCommandline);
+		}
+		if (Params.IterativeDeploy)
+		{
+			// push over the changed files
+			RunAndLog(CmdEnv, DeployServer, "Deploy -manifest \"" + CombinePaths(Params.BaseStageDirectory, "IOS", "Manifest_DeltaUFSFiles.txt") + "\"" + (String.IsNullOrEmpty(Params.Device) ? "" : " -device " + Params.Device.Substring(4)) + AdditionalCommandline + " -bundle " + BundleIdentifier);
 		}
+		Directory.SetCurrentDirectory (CurrentDir);
         PrintRunTime();
     }
 
@@ -852,7 +972,7 @@ public class IOSPlatform : Platform
 		}
 		else
 		{
-			// get the CFBundleIdentifier and modify the command line
+/*			// get the CFBundleIdentifier and modify the command line
 			int Pos = ClientCmdLine.IndexOf("-Exe=") + 6;
 			int EndPos = ClientCmdLine.IndexOf("-Targetplatform=") - 2;
 			string Exe = ClientCmdLine.Substring(Pos, EndPos - Pos);
@@ -876,7 +996,8 @@ public class IOSPlatform : Platform
 				}
 			}
 
-			return base.RunClient(ClientRunFlags, ClientApp, ClientCmdLine, Params);
+			return base.RunClient(ClientRunFlags, ClientApp, ClientCmdLine, Params);*/
+			return null;
 		}
 	}
 
@@ -890,13 +1011,22 @@ public class IOSPlatform : Platform
 		get { return true; }
 	}
 
+	public override List<string> GetFilesForCRCCheck()
+	{
+		List<string> FileList = base.GetFilesForCRCCheck();
+		FileList.Add("Info.plist");
+		return FileList;
+	}
+
 	#region Hooks
 
 	public override void PreBuildAgenda(UE4Build Build, UE4Build.BuildAgenda Agenda)
 	{
 		if (UnrealBuildTool.BuildHostPlatform.Current.Platform != UnrealTargetPlatform.Mac)
 		{
-			Agenda.DotNetProjects.Add (@"Engine\Source\Programs\IOS\iPhonePackager\iPhonePackager.csproj");
+			Agenda.DotNetProjects.Add(@"Engine\Source\Programs\IOS\MobileDeviceInterface\MobileDeviceInterface.csproj");
+			Agenda.DotNetProjects.Add(@"Engine\Source\Programs\IOS\iPhonePackager\iPhonePackager.csproj");
+			Agenda.DotNetProjects.Add(@"Engine\Source\Programs\IOS\DeploymentServer\DeploymentServer.csproj");
 		}
 	}
 
diff --git a/Engine/Source/Programs/AutomationTool/Platform.cs b/Engine/Source/Programs/AutomationTool/Platform.cs
--- a/Engine/Source/Programs/AutomationTool/Platform.cs
+++ b/Engine/Source/Programs/AutomationTool/Platform.cs
@@ -402,6 +402,16 @@ namespace AutomationTool
 			get { return false; }
 		}
 
+		public virtual List<string> GetFilesForCRCCheck()
+		{
+			string CmdLine = "UE4CommandLine.txt";
+			if (DeployLowerCaseFilenames(true))
+			{
+				CmdLine = CmdLine.ToLowerInvariant();
+			}
+			return new List<string>() { CmdLine };
+		}
+
 		#region Hooks
 
 		public virtual void PreBuildAgenda(UE4Build Build, UE4Build.BuildAgenda Agenda)
@@ -436,6 +446,13 @@ namespace AutomationTool
             return false;
         }
 
+		public virtual bool RetrieveDeployedManifests(ProjectParams Params, DeploymentContext SC, out List<string> UFSManifests, out List<string> NonUFSManifests)
+		{
+			UFSManifests = null;
+			NonUFSManifests = null;
+			return false;
+		}
+
 		#endregion
 
 		#region Utilities
diff --git a/Engine/Source/Programs/AutomationTool/ProjectParams.cs b/Engine/Source/Programs/AutomationTool/ProjectParams.cs
--- a/Engine/Source/Programs/AutomationTool/ProjectParams.cs
+++ b/Engine/Source/Programs/AutomationTool/ProjectParams.cs
@@ -243,6 +243,7 @@ namespace AutomationTool
 			this.RunCommandline = InParams.RunCommandline;
 			this.Package = InParams.Package;
 			this.Deploy = InParams.Deploy;
+			this.IterativeDeploy = InParams.IterativeDeploy;
 			this.Device = InParams.Device;
 			this.ServerDevice = InParams.ServerDevice;
             this.NullRHI = InParams.NullRHI;
@@ -363,7 +364,8 @@ namespace AutomationTool
             bool? OBBinAPK = null,
             bool? Prebuilt = null,
             int? RunTimeoutSeconds = null,
-			string OverrideMinimumOS = null
+			string OverrideMinimumOS = null,
+            bool? IterativeDeploy = null
 			)
 		{
 			//
@@ -491,7 +493,8 @@ namespace AutomationTool
 			this.RunCommandline = ParseParamValueIfNotSpecified(Command, RunCommandline, "addcmdline");
 			this.Package = GetParamValueIfNotSpecified(Command, Package, this.Package, "package");
 			this.Deploy = GetParamValueIfNotSpecified(Command, Deploy, this.Deploy, "deploy");
-			this.Device = ParseParamValueIfNotSpecified(Command, Device, "device", String.Empty).Trim(new char[]{'\"'});
+			this.IterativeDeploy = GetParamValueIfNotSpecified(Command, IterativeDeploy, this.IterativeDeploy, "iterativedeploy", "iterate");
+			this.Device = ParseParamValueIfNotSpecified(Command, Device, "device", String.Empty).Trim(new char[] { '\"' });
 			this.ServerDevice = ParseParamValueIfNotSpecified(Command, ServerDevice, "serverdevice", this.Device);
 			this.NullRHI = GetParamValueIfNotSpecified(Command, NullRHI, this.NullRHI, "nullrhi");
 			this.FakeClient = GetParamValueIfNotSpecified(Command, FakeClient, this.FakeClient, "fakeclient");
@@ -930,6 +933,12 @@ namespace AutomationTool
 		[Help( "iterativecooking", "Uses the iterative cooking, command line: -iterativecooking or -iterate" )]
 		public bool IterativeCooking;
 
+		/// <summary>
+		/// Cook: Uses the iterative deploy, command line: -iterativedeploy or -iterate
+		/// </summary>
+		[Help("iterativecooking", "Uses the iterative cooking, command line: -iterativedeploy or -iterate")]
+		public bool IterativeDeploy;
+
 		#endregion
 
 		#region Stage
@@ -1760,6 +1769,7 @@ namespace AutomationTool
 				CommandUtils.Log("IsCodeBasedProject={0}", IsCodeBasedProject.ToString());
 				CommandUtils.Log("IsProgramTarget={0}", IsProgramTarget.ToString());
 				CommandUtils.Log("IterativeCooking={0}", IterativeCooking);
+				CommandUtils.Log("IterativeDeploy={0}", IterativeDeploy);
 				CommandUtils.Log("LogWindow={0}", LogWindow);
 				CommandUtils.Log("Manifests={0}", Manifests);
 				CommandUtils.Log("MapToRun={0}", MapToRun);
diff --git a/Engine/Source/Programs/AutomationTool/Scripts/CopyBuildToStagingDirectory.Automation.cs b/Engine/Source/Programs/AutomationTool/Scripts/CopyBuildToStagingDirectory.Automation.cs
--- a/Engine/Source/Programs/AutomationTool/Scripts/CopyBuildToStagingDirectory.Automation.cs
+++ b/Engine/Source/Programs/AutomationTool/Scripts/CopyBuildToStagingDirectory.Automation.cs
@@ -284,7 +284,7 @@ public partial class Project : CommandUtils
 					SC.StageFiles(StagedFileType.UFS, CombinePaths(SC.LocalRoot, "Engine/Content/Editor/Slate"), "*", true, null, null, false, !Params.UsePak(SC.StageTargetPlatform));
 				}
 				SC.StageFiles(StagedFileType.UFS, CombinePaths(SC.LocalRoot, "Engine/Content/Slate"), "*", true, null, null, false, !Params.UsePak(SC.StageTargetPlatform));
-				SC.StageFiles(StagedFileType.UFS, CombinePaths(SC.ProjectRoot, "Content/Slate"), "*", true, null, CombinePaths(SC.RelativeProjectRootForStage, "Content/Slate"), true, !Params.UsePak(SC.StageTargetPlatform));
+				SC.StageFiles(StagedFileType.UFS, CombinePaths(SC.ProjectRoot, "Content/Slate"), "*", true, new string[] {"*.uasset"}, CombinePaths(SC.RelativeProjectRootForStage, "Content/Slate"), true, !Params.UsePak(SC.StageTargetPlatform));
 			}
             foreach (string Culture in CulturesToStage)
             {
@@ -382,14 +382,20 @@ public partial class Project : CommandUtils
 		}
 	}
 
-	public static void DumpTargetManifest(Dictionary<string, string> Mapping, string Filename)
+	public static void DumpTargetManifest(Dictionary<string, string> Mapping, string Filename, string StageDir, List<string> CRCFiles)
 	{
 		if (Mapping.Count > 0)
 		{
 			var Lines = new List<string>();
 			foreach (var Pair in Mapping)
 			{
-				string Dest = Pair.Value;
+				string TimeStamp = File.GetLastWriteTimeUtc(Pair.Key).ToString();
+				if (CRCFiles.Contains(Pair.Value))
+				{
+					byte[] FileData = File.ReadAllBytes(StageDir + "/" + Pair.Value);
+					TimeStamp = BitConverter.ToString(System.Security.Cryptography.MD5.Create().ComputeHash(FileData)).Replace("-", string.Empty);
+				}
+				string Dest = Pair.Value + "\t" + TimeStamp;
 
 				Lines.Add(Dest);
 			}
@@ -397,7 +403,7 @@ public partial class Project : CommandUtils
 		}
 	}
 
-	public static void CopyManifestFilesToStageDir(Dictionary<string, string> Mapping, string StageDir, string ManifestName)
+	public static void CopyManifestFilesToStageDir(Dictionary<string, string> Mapping, string StageDir, string ManifestName, List<string> CRCFiles)
 	{
 		string ManifestPath = "";
 		string ManifestFile = "";
@@ -418,7 +424,7 @@ public partial class Project : CommandUtils
 		}
 		if (!String.IsNullOrEmpty(ManifestPath) && Mapping.Count > 0)
 		{
-			DumpTargetManifest(Mapping, ManifestPath);
+			DumpTargetManifest(Mapping, ManifestPath, StageDir, CRCFiles);
 			if (!FileExists(ManifestPath))
 			{
 				throw new AutomationException("Failed to write manifest {0}", ManifestPath);
@@ -455,17 +461,17 @@ public partial class Project : CommandUtils
 
 	public static void CopyUsingStagingManifest(ProjectParams Params, DeploymentContext SC)
 	{
-		CopyManifestFilesToStageDir(SC.NonUFSStagingFiles, SC.StageDirectory, "NonUFSFiles");
+		CopyManifestFilesToStageDir(SC.NonUFSStagingFiles, SC.StageDirectory, "NonUFSFiles", SC.StageTargetPlatform.GetFilesForCRCCheck());
 
 		if (!Params.NoDebugInfo)
 		{
-			CopyManifestFilesToStageDir(SC.NonUFSStagingFilesDebug, SC.StageDirectory, "DebugFiles");
+			CopyManifestFilesToStageDir(SC.NonUFSStagingFilesDebug, SC.StageDirectory, "DebugFiles", SC.StageTargetPlatform.GetFilesForCRCCheck());
 		}
 
         bool bStageUnrealFileSystemFiles = !Params.CookOnTheFly && !ShouldCreatePak(Params) && !Params.FileServer;
 		if (bStageUnrealFileSystemFiles)
 		{
-			CopyManifestFilesToStageDir(SC.UFSStagingFiles, SC.StageDirectory, "UFSFiles");
+			CopyManifestFilesToStageDir(SC.UFSStagingFiles, SC.StageDirectory, "UFSFiles", SC.StageTargetPlatform.GetFilesForCRCCheck());
 		}
 	}
 
@@ -768,7 +774,7 @@ public partial class Project : CommandUtils
 	public static void ApplyStagingManifest(ProjectParams Params, DeploymentContext SC)
 	{
 		MaybeConvertToLowerCase(Params, SC);
-		if (SC.Stage && !Params.NoCleanStage && !Params.SkipStage)
+		if (SC.Stage && !Params.NoCleanStage && !Params.SkipStage && !Params.IterativeDeploy)
 		{
             try
             {
@@ -972,7 +978,11 @@ public partial class Project : CommandUtils
 
 			String ProjectFile = String.Format("{0} ", SC.ProjectArgForCommandLines);
 			Directory.CreateDirectory(GetIntermediateCommandlineDir(SC));
-			string CommandLine = String.Format("{0} {1} {2} {3}\n", ProjectFile, Params.StageCommandline.Trim(new char[] { '\"' }), Params.RunCommandline.Trim(new char[] { '\"' }), FileHostParams);
+			string CommandLine = String.Format("{0} {1} {2} {3}\n", ProjectFile, Params.StageCommandline.Trim(new char[] { '\"' }), Params.RunCommandline.Trim(new char[] { '\"' }), FileHostParams).Trim();
+			if (Params.IterativeDeploy)
+			{
+				CommandLine += " -iterative";
+			}
 			File.WriteAllText(IntermediateCmdLineFile, CommandLine);
 		}
 		else if (!Params.IsCodeBasedProject)
@@ -990,6 +1000,133 @@ public partial class Project : CommandUtils
 		WriteStageCommandline(IntermediateCmdLineFile, Params, SC);
 	}
 
+	private static Dictionary<string, string> ReadDeployedManifest(ProjectParams Params, DeploymentContext SC, List<string> ManifestList)
+	{
+		Dictionary<string, string> DeployedFiles = new Dictionary<string, string>();
+		List<string> CRCFiles = SC.StageTargetPlatform.GetFilesForCRCCheck();
+
+		// read the manifest
+		bool bContinueSearch = true;
+		foreach (string Manifest in ManifestList)
+		{
+			int FilesAdded = 0;
+			if (bContinueSearch)
+			{
+				string Data = File.ReadAllText (Manifest);
+				string[] Lines = Data.Split ('\n');
+				foreach (string Line in Lines)
+				{
+					string[] Pair = Line.Split ('\t');
+					if (Pair.Length > 1)
+					{
+						string Filename = Pair [0];
+						string TimeStamp = Pair [1];
+						FilesAdded++;
+						if (DeployedFiles.ContainsKey (Filename))
+						{
+							if ((CRCFiles.Contains (Filename) && DeployedFiles [Filename] != TimeStamp) || (!CRCFiles.Contains (Filename) && DateTime.Parse (DeployedFiles [Filename]) > DateTime.Parse (TimeStamp)))
+							{
+								DeployedFiles [Filename] = TimeStamp;
+							}
+						}
+						else
+						{
+							DeployedFiles.Add (Filename, TimeStamp);
+						}
+					}
+				}
+			}
+			File.Delete(Manifest);
+
+			if (FilesAdded == 0 && bContinueSearch)
+			{
+				// no files have been deployed at all to this guy, so remove all previously added files and exit the loop as this means we need to deploy everything
+				DeployedFiles.Clear ();
+				bContinueSearch = false;
+			}
+		}
+
+		return DeployedFiles;
+	}
+
+	protected static Dictionary<string, string> ReadStagedManifest(ProjectParams Params, DeploymentContext SC, string Manifest)
+	{
+		Dictionary <string, string> StagedFiles = new Dictionary<string, string>();
+		List<string> CRCFiles = SC.StageTargetPlatform.GetFilesForCRCCheck();
+
+		// get the staged manifest from staged directory
+		string ManifestFile = SC.StageDirectory + "/" + Manifest;
+		if (File.Exists(ManifestFile))
+		{
+			string Data = File.ReadAllText(ManifestFile);
+			string[] Lines = Data.Split('\n');
+			foreach (string Line in Lines)
+			{
+				string[] Pair = Line.Split('\t');
+				if (Pair.Length > 1)
+				{
+					string Filename = Pair[0];
+					string TimeStamp = Pair[1];
+					if (!StagedFiles.ContainsKey(Filename))
+					{
+						StagedFiles.Add(Filename, TimeStamp);
+					}
+					else if ((CRCFiles.Contains(Filename) && StagedFiles[Filename] != TimeStamp) || (!CRCFiles.Contains(Filename) && DateTime.Parse(StagedFiles[Filename]) > DateTime.Parse(TimeStamp)))
+					{
+						StagedFiles[Filename] = TimeStamp;
+					}
+				}
+			}
+		}
+		return StagedFiles;
+	}
+
+	protected static void WriteDeltaManifest(ProjectParams Params, DeploymentContext SC, Dictionary<string, string> DeployedFiles, Dictionary<string, string> StagedFiles, string DeltaManifest)
+	{
+		List<string> CRCFiles = SC.StageTargetPlatform.GetFilesForCRCCheck();
+		List<string> DeltaFiles = new List<string>();
+		foreach (KeyValuePair<string, string> File in StagedFiles)
+		{
+			bool bNeedsDeploy = true;
+			if (DeployedFiles.ContainsKey(File.Key))
+			{
+				if (CRCFiles.Contains(File.Key))
+				{
+					bNeedsDeploy = (File.Value != DeployedFiles[File.Key]);
+				}
+				else
+				{
+					DateTime Staged = DateTime.Parse(File.Value);
+					DateTime Deployed = DateTime.Parse(DeployedFiles[File.Key]);
+					bNeedsDeploy = (Staged > Deployed);
+				}
+			}
+
+			if (bNeedsDeploy)
+			{
+				DeltaFiles.Add(File.Key);
+			}
+		}
+
+		// add the manifest
+		if (!DeltaManifest.Contains("NonUFS"))
+		{
+			DeltaFiles.Add("Manifest_NonUFSFiles.txt");
+			DeltaFiles.Add("Manifest_UFSFiles.txt");
+		}
+
+		// TODO: determine files which need to be removed
+
+		// write out to the deltamanifest.json
+		string ManifestFile = CombinePaths(SC.StageDirectory, DeltaManifest);
+		StreamWriter Writer = System.IO.File.CreateText(ManifestFile);
+		foreach (string Line in DeltaFiles)
+		{
+			Writer.WriteLine(Line);
+		}
+		Writer.Close();
+	}
+
 	#endregion
 
 	#region Stage Command
@@ -1105,6 +1242,28 @@ public partial class Project : CommandUtils
 					WriteStageCommandline(Params, SC);
 					CreateStagingManifest(Params, SC);
 					ApplyStagingManifest(Params, SC);
+
+					if (Params.IterativeDeploy)
+					{
+						// get the deployed file data
+						Dictionary<string, string> DeployedUFSFiles = new Dictionary<string, string>();
+						Dictionary<string, string> DeployedNonUFSFiles = new Dictionary<string, string>();
+						List<string> UFSManifests;
+						List<string> NonUFSManifests; 
+						if (SC.StageTargetPlatform.RetrieveDeployedManifests(Params, SC, out UFSManifests, out NonUFSManifests))
+						{
+							DeployedUFSFiles = ReadDeployedManifest(Params, SC, UFSManifests);
+							DeployedNonUFSFiles = ReadDeployedManifest(Params, SC, NonUFSManifests);
+						}
+
+						// get the staged file data
+						Dictionary<string, string> StagedUFSFiles = ReadStagedManifest(Params, SC, "Manifest_UFSFiles.txt");
+						Dictionary<string, string> StagedNonUFSFiles = ReadStagedManifest(Params, SC, "Manifest_NonUFSFiles.txt");
+
+						// write out the delta file data
+						WriteDeltaManifest(Params, SC, DeployedUFSFiles, StagedUFSFiles, "Manifest_DeltaUFSFiles.txt");
+						WriteDeltaManifest(Params, SC, DeployedNonUFSFiles, StagedNonUFSFiles, "Manifest_DeltaNonUFSFiles.txt");
+					}
 				}
 			}
 
diff --git a/Engine/Source/Programs/IOS/DeploymentInterface/DeploymentInterface.cs b/Engine/Source/Programs/IOS/DeploymentInterface/DeploymentInterface.cs
--- a/Engine/Source/Programs/IOS/DeploymentInterface/DeploymentInterface.cs
+++ b/Engine/Source/Programs/IOS/DeploymentInterface/DeploymentInterface.cs
@@ -40,7 +40,9 @@ namespace iPhonePackager
         bool UninstallIPAOnDevice(string ApplicationIdentifier);
         ConnectedDeviceInfo [] EnumerateConnectedDevices();
         bool InstallIPAOnDevice(string IPAPath);
+		bool InstallFilesOnDevice(string BundleIdentifier, string ManifestFile);
 		bool BackupDocumentsDirectory(string BundleIdentifier, string DestinationDocumentsDirectory);
+		bool BackupFiles(string BundleIdentifier, string[] Files);
 		bool ListApplications();
 	}
 }
diff --git a/Engine/Source/Programs/IOS/DeploymentServer/DeployTime.cs b/Engine/Source/Programs/IOS/DeploymentServer/DeployTime.cs
--- a/Engine/Source/Programs/IOS/DeploymentServer/DeployTime.cs
+++ b/Engine/Source/Programs/IOS/DeploymentServer/DeployTime.cs
@@ -86,6 +86,13 @@ namespace DeploymentServer
             DeviceTypeMapping.Add("iPad2,3", "iPad 2 (with 3G)"); // CDMA/Verizon
 
             Console.WriteLine("[deploy] Created deployment server.");
+
+			// Initialize the mobile device manager
+			if (!bHaveRegisteredHandlers)
+			{
+				Manzana.MobileDeviceInstanceManager.Initialize(MobileDeviceConnected, MobileDeviceDisconnected);
+				bHaveRegisteredHandlers = true;
+			}
         }
 
         /// <summary>
@@ -166,7 +173,7 @@ namespace DeploymentServer
                     int TotalDurationMS = 5000;
                     while (!MobileDeviceInstanceManager.AreAnyDevicesConnected() && (TotalDurationMS > 0))
                     {
-                        System.Threading.Thread.Sleep(SleepDurationMS);
+						System.Threading.Thread.Sleep(SleepDurationMS);
                         TotalDurationMS -= SleepDurationMS;
                     }
 
@@ -330,30 +337,67 @@ namespace DeploymentServer
             });
         }
 
+		/// <summary>
+		/// Installs an IPA to all connected devices
+		/// </summary>
+		public bool InstallFilesOnDevice(string BundleIdentifier, string Manifest)
+		{
+			// Transfer to all connected devices
+			return PerformActionOnAllDevices(StandardEnumerationDelayMS, delegate(MobileDeviceInstance Device)
+			{
+				bool bResult = false;
+				if (String.IsNullOrEmpty(DeviceId) || Device.DeviceId == DeviceId)
+				{
+					bResult = Device.TryCopy(BundleIdentifier, Manifest);
+					ReportIF.Log("");
+				}
+				return bResult;
+			});
+		}
+
         public bool BackupDocumentsDirectory(string BundleIdentifier, string DestinationDocumentsDirectory)
         {
             return PerformActionOnAllDevices(StandardEnumerationDelayMS, delegate(MobileDeviceInstance Device)
             {
-                string SafeDeviceName = MobileDeviceInstance.SanitizePathNoFilename(Device.DeviceName);
-
-                // Destination folder
-                string TargetFolder = Path.Combine(DestinationDocumentsDirectory, SafeDeviceName);
+				bool bResult = false;
+				if (String.IsNullOrEmpty(DeviceId) || Device.DeviceId == DeviceId)
+				{
+					string SafeDeviceName = MobileDeviceInstance.SanitizePathNoFilename(Device.DeviceName);
 
-                // Source folder
-                string SourceFolder = "/Documents/";
+					// Destination folder
+					string TargetFolder = Path.Combine(DestinationDocumentsDirectory, SafeDeviceName);
 
-                bool bResult = Device.TryBackup(BundleIdentifier, SourceFolder, TargetFolder);
+					// Source folder
+					string SourceFolder = "/Documents/";
 
-                ReportIF.Log("");
+					bResult = Device.TryBackup(BundleIdentifier, SourceFolder, TargetFolder);
 
+					ReportIF.Log("");
+				}
                 return bResult;
             });
         }
 
+		public bool BackupFiles(string BundleIdentifier, string[] DestionationFiles)
+		{
+			return PerformActionOnAllDevices(StandardEnumerationDelayMS, delegate(MobileDeviceInstance Device)
+			{
+				bool bResult = false;
+				if (String.IsNullOrEmpty(DeviceId) || Device.DeviceId == DeviceId)
+				{
+					bResult = Device.TryBackup(BundleIdentifier, DestionationFiles);
+					ReportIF.Log("");
+				}
+
+				return bResult;
+			});
+		}
+
 		public bool ListApplications()
 		{
 			return PerformActionOnAllDevices(StandardEnumerationDelayMS, delegate(MobileDeviceInstance Device)
 			{
+					Console.WriteLine("Device '{0}' has the following applications:", Device.DeviceName);
 				Device.DumpInstalledApplications();
 
 				ReportIF.Log("");
diff --git a/Engine/Source/Programs/IOS/DeploymentServer/Program.cs b/Engine/Source/Programs/IOS/DeploymentServer/Program.cs
--- a/Engine/Source/Programs/IOS/DeploymentServer/Program.cs
+++ b/Engine/Source/Programs/IOS/DeploymentServer/Program.cs
@@ -7,9 +7,31 @@ using System.Runtime.Remoting;
 using System.Runtime.Remoting.Channels;
 using System.Runtime.Remoting.Channels.Ipc;
 using System.Diagnostics;
+using System.Runtime.InteropServices;
+using System.Collections.Generic;
 
 namespace DeploymentServer
 {
+	internal class CoreFoundation
+	{
+		[DllImport("/System/Library/Frameworks/CoreFoundation.framework/CoreFoundation")]
+		public extern static IntPtr CFStringCreateWithCString(IntPtr allocator, string value, int encoding);
+
+		[DllImport("/System/Library/Frameworks/CoreFoundation.framework/CoreFoundation")]
+		public extern static void CFRunLoopRun();
+
+		[DllImport("/System/Library/Frameworks/CoreFoundation.framework/CoreFoundation")]
+		public extern static IntPtr CFRunLoopGetMain();
+
+		[DllImport("/System/Library/Frameworks/CoreFoundation.framework/CoreFoundation")]
+		public extern static IntPtr CFRunLoopGetCurrent();
+
+		[DllImport("/System/Library/Frameworks/CoreFoundation.framework/CoreFoundation")]
+		public extern static int CFRunLoopRunInMode(IntPtr mode, double seconds, int returnAfterSourceHandled);
+
+		public static IntPtr kCFRunLoopDefaultMode = CFStringCreateWithCString(IntPtr.Zero, "kCFRunLoopDefaultMode", 0);
+	}
+
     class Program
     {
         static void Main(string[] args)
@@ -39,20 +61,137 @@ namespace DeploymentServer
             }
             else
             {
-                // Run directly by some intrepid explorer
-                Console.WriteLine("Note: This program should only be started by iPhonePackager");
-                Console.WriteLine("  This program cannot be used on it's own.");
-
-
-                DeploymentImplementation Deployer = new DeploymentImplementation();
-                var DeviceList = Deployer.EnumerateConnectedDevices();
-                foreach (var Device in DeviceList)
-                {
-                    Console.WriteLine("  - Found device named {0} of type {1} with UDID {2}", Device.DeviceName, Device.DeviceType, Device.UDID);
-                }
+				// Run directly by some intrepid explorer
+				Console.WriteLine ("Note: This program should only be started by iPhonePackager");
+				Console.WriteLine ("  This program cannot be used on it's own.");
 
+				// Parse the command
+				if (ParseCommand(args))
+				{
+					Deployer = new DeploymentImplementation();
+					bool bCommandComplete = false;
+					if (Environment.OSVersion.Platform == PlatformID.MacOSX || Environment.OSVersion.Platform == PlatformID.Unix)
+					{
+						System.Threading.Thread enumerateLoop = new System.Threading.Thread(delegate()
+						{
+							RunCommand();
+							bCommandComplete = true;
+						});
+						enumerateLoop.Start();
+						while (!bCommandComplete)
+						{
+							CoreFoundation.CFRunLoopRunInMode(CoreFoundation.kCFRunLoopDefaultMode, 1.0, 0);
+						}
+					}
+					else
+					{
+						RunCommand();
+					}
+				}
                 Console.WriteLine("Exiting.");
             }
         }
+
+		private static string Command = "";
+		private static List<string> FileList = new List<string>();
+		private static string Bundle = "";
+		private static string Manifest = "";
+		private static string ipaPath = "";
+		private static string Device = "";
+		private static bool ParseCommand(string[] Arguments)
+		{
+			if (Arguments.Length >= 1)
+			{
+				Command = Arguments[0].ToLowerInvariant();
+				for (int ArgIndex = 1; ArgIndex < Arguments.Length; ArgIndex++)
+				{
+					string Arg = Arguments[ArgIndex].ToLowerInvariant();
+					if (Arg.StartsWith("-"))
+					{
+						switch (Arg)
+						{
+							case "-file":
+								if (Arguments.Length > ArgIndex + 1)
+								{
+									FileList.Add(Arguments[++ArgIndex]);
+								}
+								else
+								{
+									return false;
+								}
+								break;
+
+							case "-bundle":
+								if (Arguments.Length > ArgIndex + 1)
+								{
+									Bundle = Arguments[++ArgIndex];
+								}
+								else
+								{
+									return false;
+								}
+								break;
+
+							case "-manifest":
+								if (Arguments.Length > ArgIndex + 1)
+								{
+									Manifest = Arguments[++ArgIndex];
+								}
+								else
+								{
+									return false;
+								}
+								break;
+
+							case "-ipa":
+								if (Arguments.Length > ArgIndex + 1)
+								{
+									ipaPath = Arguments[++ArgIndex];
+								}
+								else
+								{
+									return false;
+								}
+								break;
+
+							case "-device":
+								if (Arguments.Length > ArgIndex + 1)
+								{
+									Device = Arguments[++ArgIndex];
+								}
+								else
+								{
+									return false;
+								}
+								break;
+						}
+					}
+				}
+			}
+			return true;
+		}
+
+		private static DeploymentImplementation Deployer = null;
+		private static void RunCommand()
+		{
+			if (!String.IsNullOrEmpty(Device) && !Device.Contains("All_iOS_On"))
+			{
+				Deployer.DeviceId = Device;
+			}
+			switch (Command)
+			{
+				case "backup":
+					Deployer.BackupFiles(Bundle, FileList.ToArray());
+					break;
+
+				case "deploy":
+					Deployer.InstallFilesOnDevice(Bundle, Manifest);
+					break;
+
+				case "install":
+					Deployer.InstallIPAOnDevice(ipaPath);
+					break;
+			}
+		}
     }
 }
diff --git a/Engine/Source/Programs/IOS/MobileDeviceInterface/CoreFoundation.cs b/Engine/Source/Programs/IOS/MobileDeviceInterface/CoreFoundation.cs
--- a/Engine/Source/Programs/IOS/MobileDeviceInterface/CoreFoundation.cs
+++ b/Engine/Source/Programs/IOS/MobileDeviceInterface/CoreFoundation.cs
@@ -19,6 +19,46 @@ namespace Manzana
     public class CFNumber {  }
     public class CFBoolean { }
 
+	public enum CFStringBuiltInEncodings
+	{
+		kCFStringEncodingMacRoman = 0x00000000,
+		kCFStringEncodingWindowsLatin1 = 0x00000500,
+		kCFStringEncodingISOLatin1 = 0x00000201,
+		kCFStringEncodingNextStepLatin = 0x00000B01,
+		kCFStringEncodingASCII = 0x00000600,
+		kCFStringEncodingUnicode = 0x00000100,
+		kCFStringEncodingUTF8 = 0x08000100,
+		kCFStringEncodingNonLossyASCII = 0x00000BFF
+	};
+
+	public enum CFURLPathStyle
+	{
+		kCFURLPOSIXPathStyle = 0,
+		kCFURLHFSPathStyle = 1,
+		kCFURLWindowsPathStyle = 2
+	};
+
+	public enum CFNumberType
+	{
+		kCFNumberSInt8Type = 1,
+		kCFNumberSInt16Type = 2,
+		kCFNumberSInt32Type = 3,
+		kCFNumberSInt64Type = 4,
+		kCFNumberFloat32Type = 5,
+		kCFNumberFloat64Type = 6,
+		kCFNumberCharType = 7,
+		kCFNumberShortType = 8,
+		kCFNumberIntType = 9,
+		kCFNumberLongType = 10,
+		kCFNumberLongLongType = 11,
+		kCFNumberFloatType = 12,
+		kCFNumberDoubleType = 13,
+		kCFNumberCFIndexType = 14,
+		kCFNumberNSIntegerType = 15,
+		kCFNumberCGFloatType = 16,
+		kCFNumberMaxType = 16
+	};
+
     internal class RequiredWinAPI
     {
         [DllImport("kernel32", CharSet = CharSet.Auto, SetLastError = true)]
@@ -31,6 +71,412 @@ namespace Manzana
         public static extern IntPtr GetModuleHandle(string lpModuleName);
     }
 
+	internal interface CoreFoundationImpl
+	{
+		TypedPtr<CFString> __CFStringMakeConstantString(byte[] s);
+		TypedPtr<CFURL> CFURLCreateWithFileSystemPath(IntPtr Allocator, TypedPtr<CFString> FilePath, CFURLPathStyle PathStyle, int isDirectory);
+		Boolean CFStringGetCString(TypedPtr<CFString> theString, byte[] buffer, int bufferSize, CFStringBuiltInEncodings encoding);
+		TypedPtr<CFString> CFURLGetString(IntPtr anURL);
+		uint CFGetTypeID(IntPtr FromInstance);
+		uint CFStringGetTypeID();
+		uint CFNumberGetTypeID();
+		uint CFBooleanGetTypeID();
+		uint CFDictionaryGetTypeID();
+		int CFDictionaryGetCount (TypedPtr<CFDictionary> Dict);
+		IntPtr CFNumberCreate(IntPtr allocator, CFNumberType theType, ref Int32 Value);
+		Boolean CFNumberGetValue(TypedPtr<CFNumber> Number, CFNumberType TheType, out float Value);
+		Boolean CFNumberGetValue(TypedPtr<CFNumber> Number, CFNumberType TheType, out double Value);
+		Boolean CFNumberGetValue(TypedPtr<CFNumber> Number, CFNumberType TheType, out int Value);
+		Boolean CFNumberGetValue(TypedPtr<CFNumber> Number, CFNumberType TheType, out uint Value);
+		Boolean CFNumberGetValue(TypedPtr<CFNumber> Number, CFNumberType TheType, out Int64 Value);
+		CFNumberType CFNumberGetType(TypedPtr<CFNumber> Number);
+		void CFPreferencesSetAppValue(IntPtr Key, IntPtr Value, IntPtr ApplicationID);
+		void CFDictionarySetValue(TypedPtr<CFDictionary> Dict, /*const*/ IntPtr Key, /*const*/ IntPtr Value);
+		void CFDictionaryGetKeysAndValues(TypedPtr<CFDictionary> Dict, /*const*/ IntPtr[] Keys, /*const*/ IntPtr[] Values);
+		TypedPtr<CFDictionary> CFDictionaryCreateMutable(
+			/* CFAllocatorRef */ IntPtr allocator,
+			/* CFIndex */ int capacity,
+			/* const CFDictionaryKeyCallBacks* */ IntPtr keyCallBacks,
+			/* const CFDictionaryValueCallBacks* */ IntPtr valueCallBacks
+		);
+	}
+
+	internal class CoreFoundationWin : CoreFoundationImpl
+	{
+
+		public TypedPtr<CFString> __CFStringMakeConstantString(byte[] s)
+		{
+			return CoreFoundation.__CFStringMakeConstantString(s);
+		}
+
+		public TypedPtr<CFURL> CFURLCreateWithFileSystemPath(IntPtr Allocator, TypedPtr<CFString> FilePath, CFURLPathStyle PathStyle, int isDirectory)
+		{
+			return CoreFoundation.CFURLCreateWithFileSystemPath(Allocator, (IntPtr)FilePath, PathStyle, isDirectory);
+		}
+
+		public Boolean CFStringGetCString(TypedPtr<CFString> theString, byte[] buffer, int bufferSize, CFStringBuiltInEncodings encoding)
+		{
+			return CoreFoundation.CFStringGetCString((IntPtr)theString, buffer, bufferSize, encoding);
+		}
+
+		public TypedPtr<CFString> CFURLGetString(IntPtr anURL)
+		{
+			return CoreFoundation.CFURLGetString(anURL);
+		}
+
+		public uint CFGetTypeID(IntPtr FromInstance)
+		{
+			return CoreFoundation.CFGetTypeID(FromInstance);
+		}
+
+		public uint CFStringGetTypeID()
+		{
+			return CoreFoundation.CFStringGetTypeID();
+		}
+
+		public uint CFNumberGetTypeID()
+		{
+			return CoreFoundation.CFNumberGetTypeID();
+		}
+
+		public uint CFBooleanGetTypeID()
+		{
+			return CoreFoundation.CFBooleanGetTypeID();
+		}
+
+		public uint CFDictionaryGetTypeID()
+		{
+			return CoreFoundation.CFDictionaryGetTypeID();
+		}
+
+		public int CFDictionaryGetCount(TypedPtr<CFDictionary> Dict)
+		{
+			return CoreFoundation.CFDictionaryGetCount ((IntPtr)Dict);
+		}
+
+		public IntPtr CFNumberCreate(IntPtr allocator, CFNumberType theType, ref Int32 Value)
+		{
+			return CoreFoundation.CFNumberCreate(allocator, theType, ref Value);
+		}
+
+		public Boolean CFNumberGetValue(TypedPtr<CFNumber> Number, CFNumberType TheType, out float Value)
+		{
+			return CoreFoundation.CFNumberGetValue((IntPtr)Number, TheType, out Value);
+		}
+
+		public Boolean CFNumberGetValue(TypedPtr<CFNumber> Number, CFNumberType TheType, out double Value)
+		{
+			return CoreFoundation.CFNumberGetValue((IntPtr)Number, TheType, out Value);
+		}
+
+		public Boolean CFNumberGetValue(TypedPtr<CFNumber> Number, CFNumberType TheType, out int Value)
+		{
+			return CoreFoundation.CFNumberGetValue((IntPtr)Number, TheType, out Value);
+		}
+
+		public Boolean CFNumberGetValue(TypedPtr<CFNumber> Number, CFNumberType TheType, out uint Value)
+		{
+			return CoreFoundation.CFNumberGetValue((IntPtr)Number, TheType, out Value);
+		}
+
+		public Boolean CFNumberGetValue(TypedPtr<CFNumber> Number, CFNumberType TheType, out Int64 Value)
+		{
+			return CoreFoundation.CFNumberGetValue((IntPtr)Number, TheType, out Value);
+		}
+
+		public CFNumberType CFNumberGetType(TypedPtr<CFNumber> Number)
+		{
+			return CoreFoundation.CFNumberGetType((IntPtr)Number);
+		}
+
+		public void CFPreferencesSetAppValue(IntPtr Key, IntPtr Value, IntPtr ApplicationID)
+		{
+			CoreFoundation.CFPreferencesSetAppValue(Key, Value, ApplicationID);
+		}
+
+		public void CFDictionarySetValue(TypedPtr<CFDictionary> Dict, /*const*/ IntPtr Key, /*const*/ IntPtr Value)
+		{
+			CoreFoundation.CFDictionarySetValue((IntPtr)Dict, Key, Value);
+		}
+
+		public void CFDictionaryGetKeysAndValues(TypedPtr<CFDictionary> Dict, /*const*/ IntPtr[] Keys, /*const*/ IntPtr[] Values)
+		{
+			CoreFoundation.CFDictionaryGetKeysAndValues((IntPtr)Dict, Keys, Values);
+		}
+
+		public TypedPtr<CFDictionary> CFDictionaryCreateMutable(
+			/* CFAllocatorRef */ IntPtr allocator,
+			/* CFIndex */ int capacity,
+			/* const CFDictionaryKeyCallBacks* */ IntPtr keyCallBacks,
+			/* const CFDictionaryValueCallBacks* */ IntPtr valueCallBacks
+		)
+		{
+			return CoreFoundation.CFDictionaryCreateMutable(allocator, capacity, keyCallBacks, valueCallBacks);
+		}
+
+		private class CoreFoundation
+		{
+			[DllImport("CoreFoundation.dll", CallingConvention = CallingConvention.Cdecl)]
+			public static extern IntPtr __CFStringMakeConstantString(byte[] s);
+
+			[DllImport("CoreFoundation.dll", CallingConvention = CallingConvention.Cdecl)]
+			public static extern IntPtr/*CFURL*/ CFURLCreateWithFileSystemPath(IntPtr Allocator, IntPtr/*CFString*/ FilePath, CFURLPathStyle PathStyle, int isDirectory);
+
+			[DllImport("CoreFoundation.dll", CallingConvention = CallingConvention.Cdecl)]
+			public static extern IntPtr/*CFString*/ CFURLGetString(IntPtr anURL);
+
+			[DllImport("CoreFoundation.dll", CallingConvention = CallingConvention.Cdecl)]
+			public static extern uint CFGetTypeID(IntPtr FromInstance);
+			[DllImport("CoreFoundation.dll", CallingConvention = CallingConvention.Cdecl)]
+			public static extern uint CFStringGetTypeID();
+			[DllImport("CoreFoundation.dll", CallingConvention = CallingConvention.Cdecl)]
+			public static extern uint CFNumberGetTypeID();
+			[DllImport("CoreFoundation.dll", CallingConvention = CallingConvention.Cdecl)]
+			public static extern uint CFBooleanGetTypeID();
+			[DllImport("CoreFoundation.dll", CallingConvention = CallingConvention.Cdecl)]
+			public static extern uint CFDictionaryGetTypeID();
+
+			[DllImport("CoreFoundation.dll", CallingConvention = CallingConvention.Cdecl)]
+			public static extern void CFDictionaryAddValue(IntPtr/*CFDictionary*/ Dict, /*const*/ IntPtr Key, /*const*/ IntPtr Value);
+
+			[DllImport("CoreFoundation.dll", CallingConvention = CallingConvention.Cdecl)]
+			public static extern void CFDictionarySetValue(IntPtr/*CFDictionary*/ Dict, /*const*/ IntPtr Key, /*const*/ IntPtr Value);
+
+			[DllImport("CoreFoundation.dll", CallingConvention = CallingConvention.Cdecl)]
+			public static extern void CFDictionaryGetKeysAndValues(IntPtr/*CFDictionary*/ Dict, /*const*/ IntPtr[] Keys, /*const*/ IntPtr[] Values);
+
+			[DllImport("CoreFoundation.dll", CallingConvention = CallingConvention.Cdecl)]
+			public static extern IntPtr/*CFDictionary*/ CFDictionaryCreateMutable(
+				/* CFAllocatorRef */ IntPtr allocator,
+				/* CFIndex */ int capacity,
+				/* const CFDictionaryKeyCallBacks* */ IntPtr keyCallBacks,
+				/* const CFDictionaryValueCallBacks* */ IntPtr valueCallBacks
+			);
+
+			[DllImport("CoreFoundation.dll", CallingConvention = CallingConvention.Cdecl)]
+			public static extern int CFDictionaryGetCount(IntPtr/*CFDictionary*/ Dict);
+
+			[DllImport("CoreFoundation.dll", CallingConvention = CallingConvention.Cdecl)]
+			public extern static Boolean CFStringGetCString(IntPtr/*CFString*/ theString, byte[] buffer, int bufferSize, CFStringBuiltInEncodings encoding);
+
+			//[DllImport("CoreFoundation.dll", CallingConvention = CallingConvention.Cdecl)]
+			//public extern static IntPtr CFNumberCreate(IntPtr allocator, CFNumberType theType, IntPtr valuePtr);
+
+			//@TODO: make this more general
+			[DllImport("CoreFoundation.dll", CallingConvention = CallingConvention.Cdecl)]
+			public extern static IntPtr CFNumberCreate(IntPtr allocator, CFNumberType theType, ref Int32 Value);
+
+			[DllImport("CoreFoundation.dll", EntryPoint = "CFNumberGetValue", CallingConvention = CallingConvention.Cdecl)]
+			public extern static Boolean CFNumberGetValue(IntPtr/*CFNumber*/ Number, CFNumberType TheType, out float Value);
+
+			[DllImport("CoreFoundation.dll", EntryPoint = "CFNumberGetValue", CallingConvention = CallingConvention.Cdecl)]
+			public extern static Boolean CFNumberGetValue(IntPtr/*CFNumber*/ Number, CFNumberType TheType, out double Value);
+
+			[DllImport("CoreFoundation.dll", EntryPoint = "CFNumberGetValue", CallingConvention = CallingConvention.Cdecl)]
+			public extern static Boolean CFNumberGetValue(IntPtr/*CFNumber*/ Number, CFNumberType TheType, out int Value);
+
+			[DllImport("CoreFoundation.dll", EntryPoint = "CFNumberGetValue", CallingConvention = CallingConvention.Cdecl)]
+			public extern static Boolean CFNumberGetValue(IntPtr/*CFNumber*/ Number, CFNumberType TheType, out uint Value);
+
+			[DllImport("CoreFoundation.dll", EntryPoint = "CFNumberGetValue", CallingConvention = CallingConvention.Cdecl)]
+			public extern static Boolean CFNumberGetValue(IntPtr/*CFNumber*/ Number, CFNumberType TheType, out Int64 Value);
+
+			[DllImport("CoreFoundation.dll", CallingConvention = CallingConvention.Cdecl)]
+			public extern static CFNumberType CFNumberGetType(IntPtr/*CFNumber*/ Number);
+
+			[DllImport("CoreFoundation.dll", CallingConvention = CallingConvention.Cdecl)]
+			public static extern void CFPreferencesSetAppValue(IntPtr Key, IntPtr Value, IntPtr ApplicationID);
+		}
+	}
+
+	internal class CoreFoundationOSX : CoreFoundationImpl
+	{
+
+		public TypedPtr<CFString> __CFStringMakeConstantString(byte[] s)
+		{
+			return CoreFoundation.__CFStringMakeConstantString(s);
+		}
+
+		public TypedPtr<CFURL> CFURLCreateWithFileSystemPath(IntPtr Allocator, TypedPtr<CFString> FilePath, CFURLPathStyle PathStyle, int isDirectory)
+		{
+			return CoreFoundation.CFURLCreateWithFileSystemPath(Allocator, (IntPtr)FilePath, PathStyle, isDirectory);
+		}
+
+		public Boolean CFStringGetCString(TypedPtr<CFString> theString, byte[] buffer, int bufferSize, CFStringBuiltInEncodings encoding)
+		{
+			return CoreFoundation.CFStringGetCString((IntPtr)theString, buffer, bufferSize, encoding);
+		}
+
+		public TypedPtr<CFString> CFURLGetString(IntPtr anURL)
+		{
+			return CoreFoundation.CFURLGetString(anURL);
+		}
+
+		public uint CFGetTypeID(IntPtr FromInstance)
+		{
+			return CoreFoundation.CFGetTypeID(FromInstance);
+		}
+
+		public uint CFStringGetTypeID()
+		{
+			return CoreFoundation.CFStringGetTypeID();
+		}
+
+		public uint CFNumberGetTypeID()
+		{
+			return CoreFoundation.CFNumberGetTypeID();
+		}
+
+		public uint CFBooleanGetTypeID()
+		{
+			return CoreFoundation.CFBooleanGetTypeID();
+		}
+
+		public uint CFDictionaryGetTypeID()
+		{
+			return CoreFoundation.CFDictionaryGetTypeID();
+		}
+
+		public int CFDictionaryGetCount(TypedPtr<CFDictionary> Dict)
+		{
+			return CoreFoundation.CFDictionaryGetCount ((IntPtr)Dict);
+		}
+
+		public IntPtr CFNumberCreate(IntPtr allocator, CFNumberType theType, ref Int32 Value)
+		{
+			return CoreFoundation.CFNumberCreate(allocator, theType, ref Value);
+		}
+
+		public Boolean CFNumberGetValue(TypedPtr<CFNumber> Number, CFNumberType TheType, out float Value)
+		{
+			return CoreFoundation.CFNumberGetValue((IntPtr)Number, TheType, out Value);
+		}
+
+		public Boolean CFNumberGetValue(TypedPtr<CFNumber> Number, CFNumberType TheType, out double Value)
+		{
+			return CoreFoundation.CFNumberGetValue((IntPtr)Number, TheType, out Value);
+		}
+
+		public Boolean CFNumberGetValue(TypedPtr<CFNumber> Number, CFNumberType TheType, out int Value)
+		{
+			return CoreFoundation.CFNumberGetValue((IntPtr)Number, TheType, out Value);
+		}
+
+		public Boolean CFNumberGetValue(TypedPtr<CFNumber> Number, CFNumberType TheType, out uint Value)
+		{
+			return CoreFoundation.CFNumberGetValue((IntPtr)Number, TheType, out Value);
+		}
+
+		public Boolean CFNumberGetValue(TypedPtr<CFNumber> Number, CFNumberType TheType, out Int64 Value)
+		{
+			return CoreFoundation.CFNumberGetValue((IntPtr)Number, TheType, out Value);
+		}
+
+		public CFNumberType CFNumberGetType(TypedPtr<CFNumber> Number)
+		{
+			return CoreFoundation.CFNumberGetType((IntPtr)Number);
+		}
+
+		public void CFPreferencesSetAppValue(IntPtr Key, IntPtr Value, IntPtr ApplicationID)
+		{
+			CoreFoundation.CFPreferencesSetAppValue(Key, Value, ApplicationID);
+		}
+
+		public void CFDictionarySetValue(TypedPtr<CFDictionary> Dict, /*const*/ IntPtr Key, /*const*/ IntPtr Value)
+		{
+			CoreFoundation.CFDictionarySetValue((IntPtr)Dict, Key, Value);
+		}
+
+		public void CFDictionaryGetKeysAndValues(TypedPtr<CFDictionary> Dict, /*const*/ IntPtr[] Keys, /*const*/ IntPtr[] Values)
+		{
+			CoreFoundation.CFDictionaryGetKeysAndValues((IntPtr)Dict, Keys, Values);
+		}
+
+		public TypedPtr<CFDictionary> CFDictionaryCreateMutable(
+			/* CFAllocatorRef */ IntPtr allocator,
+			/* CFIndex */ int capacity,
+			/* const CFDictionaryKeyCallBacks* */ IntPtr keyCallBacks,
+			/* const CFDictionaryValueCallBacks* */ IntPtr valueCallBacks
+		)
+		{
+			return CoreFoundation.CFDictionaryCreateMutable(allocator, capacity, keyCallBacks, valueCallBacks);
+		}
+
+		private class CoreFoundation
+		{
+			[DllImport("/System/Library/Frameworks/CoreFoundation.framework/CoreFoundation", CallingConvention = CallingConvention.Cdecl)]
+			public static extern IntPtr __CFStringMakeConstantString(byte[] s);
+
+			[DllImport("/System/Library/Frameworks/CoreFoundation.framework/CoreFoundation", CallingConvention = CallingConvention.Cdecl)]
+			public static extern IntPtr/*CFURL*/ CFURLCreateWithFileSystemPath(IntPtr Allocator, IntPtr/*CFString*/ FilePath, CFURLPathStyle PathStyle, int isDirectory);
+
+			[DllImport("/System/Library/Frameworks/CoreFoundation.framework/CoreFoundation", CallingConvention = CallingConvention.Cdecl)]
+			public static extern IntPtr/*CFString*/ CFURLGetString(IntPtr anURL);
+
+			[DllImport("/System/Library/Frameworks/CoreFoundation.framework/CoreFoundation", CallingConvention = CallingConvention.Cdecl)]
+			public static extern uint CFGetTypeID(IntPtr FromInstance);
+			[DllImport("/System/Library/Frameworks/CoreFoundation.framework/CoreFoundation", CallingConvention = CallingConvention.Cdecl)]
+			public static extern uint CFStringGetTypeID();
+			[DllImport("/System/Library/Frameworks/CoreFoundation.framework/CoreFoundation", CallingConvention = CallingConvention.Cdecl)]
+			public static extern uint CFNumberGetTypeID();
+			[DllImport("/System/Library/Frameworks/CoreFoundation.framework/CoreFoundation", CallingConvention = CallingConvention.Cdecl)]
+			public static extern uint CFBooleanGetTypeID();
+			[DllImport("/System/Library/Frameworks/CoreFoundation.framework/CoreFoundation", CallingConvention = CallingConvention.Cdecl)]
+			public static extern uint CFDictionaryGetTypeID();
+
+			[DllImport("/System/Library/Frameworks/CoreFoundation.framework/CoreFoundation", CallingConvention = CallingConvention.Cdecl)]
+			public static extern void CFDictionaryAddValue(IntPtr/*CFDictionary*/ Dict, /*const*/ IntPtr Key, /*const*/ IntPtr Value);
+
+			[DllImport("/System/Library/Frameworks/CoreFoundation.framework/CoreFoundation", CallingConvention = CallingConvention.Cdecl)]
+			public static extern void CFDictionarySetValue(IntPtr/*CFDictionary*/ Dict, /*const*/ IntPtr Key, /*const*/ IntPtr Value);
+
+			[DllImport("/System/Library/Frameworks/CoreFoundation.framework/CoreFoundation", CallingConvention = CallingConvention.Cdecl)]
+			public static extern void CFDictionaryGetKeysAndValues(IntPtr/*CFDictionary*/ Dict, /*const*/ IntPtr[] Keys, /*const*/ IntPtr[] Values);
+
+			[DllImport("/System/Library/Frameworks/CoreFoundation.framework/CoreFoundation", CallingConvention = CallingConvention.Cdecl)]
+			public static extern IntPtr/*CFDictionary*/ CFDictionaryCreateMutable(
+				/* CFAllocatorRef */ IntPtr allocator,
+				/* CFIndex */ int capacity,
+				/* const CFDictionaryKeyCallBacks* */ IntPtr keyCallBacks,
+				/* const CFDictionaryValueCallBacks* */ IntPtr valueCallBacks
+			);
+
+			[DllImport("/System/Library/Frameworks/CoreFoundation.framework/CoreFoundation", CallingConvention = CallingConvention.Cdecl)]
+			public static extern int CFDictionaryGetCount(IntPtr/*CFDictionary*/ Dict);
+
+			[DllImport("/System/Library/Frameworks/CoreFoundation.framework/CoreFoundation", CallingConvention = CallingConvention.Cdecl)]
+			public extern static Boolean CFStringGetCString(IntPtr/*CFString*/ theString, byte[] buffer, int bufferSize, CFStringBuiltInEncodings encoding);
+
+			//[DllImport("CoreFoundation.dylib", CallingConvention = CallingConvention.Cdecl)]
+			//public extern static IntPtr CFNumberCreate(IntPtr allocator, CFNumberType theType, IntPtr valuePtr);
+
+			//@TODO: make this more general
+			[DllImport("/System/Library/Frameworks/CoreFoundation.framework/CoreFoundation", CallingConvention = CallingConvention.Cdecl)]
+			public extern static IntPtr CFNumberCreate(IntPtr allocator, CFNumberType theType, ref Int32 Value);
+
+			[DllImport("/System/Library/Frameworks/CoreFoundation.framework/CoreFoundation", EntryPoint = "CFNumberGetValue", CallingConvention = CallingConvention.Cdecl)]
+			public extern static Boolean CFNumberGetValue(IntPtr/*CFNumber*/ Number, CFNumberType TheType, out float Value);
+
+			[DllImport("/System/Library/Frameworks/CoreFoundation.framework/CoreFoundation", EntryPoint = "CFNumberGetValue", CallingConvention = CallingConvention.Cdecl)]
+			public extern static Boolean CFNumberGetValue(IntPtr/*CFNumber*/ Number, CFNumberType TheType, out double Value);
+
+			[DllImport("/System/Library/Frameworks/CoreFoundation.framework/CoreFoundation", EntryPoint = "CFNumberGetValue", CallingConvention = CallingConvention.Cdecl)]
+			public extern static Boolean CFNumberGetValue(IntPtr/*CFNumber*/ Number, CFNumberType TheType, out int Value);
+
+			[DllImport("/System/Library/Frameworks/CoreFoundation.framework/CoreFoundation", EntryPoint = "CFNumberGetValue", CallingConvention = CallingConvention.Cdecl)]
+			public extern static Boolean CFNumberGetValue(IntPtr/*CFNumber*/ Number, CFNumberType TheType, out uint Value);
+
+			[DllImport("/System/Library/Frameworks/CoreFoundation.framework/CoreFoundation", EntryPoint = "CFNumberGetValue", CallingConvention = CallingConvention.Cdecl)]
+			public extern static Boolean CFNumberGetValue(IntPtr/*CFNumber*/ Number, CFNumberType TheType, out Int64 Value);
+
+			[DllImport("/System/Library/Frameworks/CoreFoundation.framework/CoreFoundation", CallingConvention = CallingConvention.Cdecl)]
+			public extern static CFNumberType CFNumberGetType(IntPtr/*CFNumber*/ Number);
+
+			[DllImport("/System/Library/Frameworks/CoreFoundation.framework/CoreFoundation", CallingConvention = CallingConvention.Cdecl)]
+			public static extern void CFPreferencesSetAppValue(IntPtr Key, IntPtr Value, IntPtr ApplicationID);
+		}
+	}
+
     /// <summary>
     /// Declarations for various methods in Apple's CoreFoundation.dll, needed for interop with the MobileDevice DLL
     /// </summary>
@@ -43,139 +489,29 @@ namespace Manzana
         public static TypedPtr<CFBoolean> kCFBooleanTrue;
         public static TypedPtr<CFBoolean> kCFBooleanFalse;
 
+		public static CoreFoundationImpl CoreImpl = null;
+
         static void InitializeCoreFoundation()
         {
-            IntPtr CoreFoundation = RequiredWinAPI.GetModuleHandle("CoreFoundation.dll");
-
-            kCFAllocatorDefault = RequiredWinAPI.GetProcAddress(CoreFoundation, "kCFAllocatorDefault");
-            kCFTypeDictionaryKeyCallBacks = RequiredWinAPI.GetProcAddress(CoreFoundation, "kCFTypeDictionaryKeyCallBacks");
-            kCFTypeDictionaryValueCallBacks = RequiredWinAPI.GetProcAddress(CoreFoundation, "kCFTypeDictionaryValueCallBacks");
-
-            kCFBooleanTrue = RequiredWinAPI.GetProcAddress(CoreFoundation, "kCFBooleanTrue");
-            kCFBooleanFalse = RequiredWinAPI.GetProcAddress(CoreFoundation, "kCFBooleanFalse");
+			if (Environment.OSVersion.Platform == PlatformID.MacOSX || Environment.OSVersion.Platform == PlatformID.Unix)
+			{
+				CoreImpl = new CoreFoundationOSX();
+			}
+			else
+			{
+				CoreImpl = new CoreFoundationWin();
+
+				IntPtr CoreFoundation = RequiredWinAPI.GetModuleHandle("CoreFoundation.dll");
+
+				kCFAllocatorDefault = RequiredWinAPI.GetProcAddress(CoreFoundation, "kCFAllocatorDefault");
+				kCFTypeDictionaryKeyCallBacks = RequiredWinAPI.GetProcAddress(CoreFoundation, "kCFTypeDictionaryKeyCallBacks");
+				kCFTypeDictionaryValueCallBacks = RequiredWinAPI.GetProcAddress(CoreFoundation, "kCFTypeDictionaryValueCallBacks");
+
+				kCFBooleanTrue = RequiredWinAPI.GetProcAddress(CoreFoundation, "kCFBooleanTrue");
+				kCFBooleanFalse = RequiredWinAPI.GetProcAddress(CoreFoundation, "kCFBooleanFalse");
+			}
         }
 
-        [DllImport("CoreFoundation.dll", CallingConvention = CallingConvention.Cdecl)]
-        private static extern IntPtr __CFStringMakeConstantString(byte[] s);
-
-        [DllImport("CoreFoundation.dll", CallingConvention = CallingConvention.Cdecl)]
-        private static extern IntPtr/*CFURL*/ CFURLCreateWithFileSystemPath(IntPtr Allocator, IntPtr/*CFString*/ FilePath, CFURLPathStyle PathStyle, int isDirectory);
-
-        [DllImport("CoreFoundation.dll", CallingConvention = CallingConvention.Cdecl)]
-        public static extern IntPtr/*CFString*/ CFURLGetString(IntPtr anURL);
-
-        [DllImport("CoreFoundation.dll", CallingConvention = CallingConvention.Cdecl)]
-        public static extern uint CFGetTypeID(IntPtr FromInstance);
-        [DllImport("CoreFoundation.dll", CallingConvention = CallingConvention.Cdecl)]
-        public static extern uint CFStringGetTypeID();
-        [DllImport("CoreFoundation.dll", CallingConvention = CallingConvention.Cdecl)]
-        public static extern uint CFNumberGetTypeID();
-        [DllImport("CoreFoundation.dll", CallingConvention = CallingConvention.Cdecl)]
-        public static extern uint CFBooleanGetTypeID();
-        [DllImport("CoreFoundation.dll", CallingConvention = CallingConvention.Cdecl)]
-        public static extern uint CFDictionaryGetTypeID();
-
-        [DllImport("CoreFoundation.dll", CallingConvention = CallingConvention.Cdecl)]
-        private static extern void CFDictionaryAddValue(IntPtr/*CFDictionary*/ Dict, /*const*/ IntPtr Key, /*const*/ IntPtr Value);
-
-        [DllImport("CoreFoundation.dll", CallingConvention = CallingConvention.Cdecl)]
-        private static extern void CFDictionarySetValue(IntPtr/*CFDictionary*/ Dict, /*const*/ IntPtr Key, /*const*/ IntPtr Value);
-
-        [DllImport("CoreFoundation.dll", CallingConvention = CallingConvention.Cdecl)]
-        private static extern void CFDictionaryGetKeysAndValues(IntPtr/*CFDictionary*/ Dict, /*const*/ IntPtr[] Keys, /*const*/ IntPtr[] Values);
-
-        [DllImport("CoreFoundation.dll", CallingConvention = CallingConvention.Cdecl)]
-        private static extern IntPtr/*CFDictionary*/ CFDictionaryCreateMutable(
-            /* CFAllocatorRef */ IntPtr allocator,
-            /* CFIndex */ int capacity,
-            /* const CFDictionaryKeyCallBacks* */ IntPtr keyCallBacks,
-            /* const CFDictionaryValueCallBacks* */ IntPtr valueCallBacks
-        );
-
-        [DllImport("CoreFoundation.dll", CallingConvention = CallingConvention.Cdecl)]
-        private static extern int CFDictionaryGetCount(IntPtr/*CFDictionary*/ Dict);
-
-        [DllImport("CoreFoundation.dll", CallingConvention = CallingConvention.Cdecl)]
-        private extern static Boolean CFStringGetCString(IntPtr/*CFString*/ theString, byte[] buffer, int bufferSize, CFStringBuiltInEncodings encoding);
-
-        //[DllImport("CoreFoundation.dll", CallingConvention = CallingConvention.Cdecl)]
-        //public extern static IntPtr CFNumberCreate(IntPtr allocator, CFNumberType theType, IntPtr valuePtr);
-
-        //@TODO: make this more general
-        [DllImport("CoreFoundation.dll", CallingConvention = CallingConvention.Cdecl)]
-        private extern static IntPtr CFNumberCreate(IntPtr allocator, CFNumberType theType, ref Int32 Value);
-
-        [DllImport("CoreFoundation.dll", EntryPoint = "CFNumberGetValue", CallingConvention = CallingConvention.Cdecl)]
-        private extern static Boolean CFNumberGetValue(IntPtr/*CFNumber*/ Number, CFNumberType TheType, out float Value);
-
-        [DllImport("CoreFoundation.dll", EntryPoint = "CFNumberGetValue", CallingConvention = CallingConvention.Cdecl)]
-        private extern static Boolean CFNumberGetValue(IntPtr/*CFNumber*/ Number, CFNumberType TheType, out double Value);
-
-        [DllImport("CoreFoundation.dll", EntryPoint = "CFNumberGetValue", CallingConvention = CallingConvention.Cdecl)]
-        private extern static Boolean CFNumberGetValue(IntPtr/*CFNumber*/ Number, CFNumberType TheType, out int Value);
-
-        [DllImport("CoreFoundation.dll", EntryPoint = "CFNumberGetValue", CallingConvention = CallingConvention.Cdecl)]
-        private extern static Boolean CFNumberGetValue(IntPtr/*CFNumber*/ Number, CFNumberType TheType, out uint Value);
-
-        [DllImport("CoreFoundation.dll", EntryPoint = "CFNumberGetValue", CallingConvention = CallingConvention.Cdecl)]
-        private extern static Boolean CFNumberGetValue(IntPtr/*CFNumber*/ Number, CFNumberType TheType, out Int64 Value);
-
-        [DllImport("CoreFoundation.dll", CallingConvention = CallingConvention.Cdecl)]
-        private extern static CFNumberType CFNumberGetType(IntPtr/*CFNumber*/ Number);
-
-        [DllImport("CoreFoundation.dll", CallingConvention = CallingConvention.Cdecl)]
-        public static extern void CFPreferencesSetAppValue(IntPtr Key, IntPtr Value, IntPtr ApplicationID);
-
-
-
-
-
-
-        public enum CFStringBuiltInEncodings
-        {
-            kCFStringEncodingMacRoman = 0x00000000,
-            kCFStringEncodingWindowsLatin1 = 0x00000500,
-            kCFStringEncodingISOLatin1 = 0x00000201,
-            kCFStringEncodingNextStepLatin = 0x00000B01,
-            kCFStringEncodingASCII = 0x00000600,
-            kCFStringEncodingUnicode = 0x00000100,
-            kCFStringEncodingUTF8 = 0x08000100,
-            kCFStringEncodingNonLossyASCII = 0x00000BFF
-        };
-
-        public enum CFURLPathStyle
-        {
-            kCFURLPOSIXPathStyle = 0,
-            kCFURLHFSPathStyle = 1,
-            kCFURLWindowsPathStyle = 2
-        };
-
-        public enum CFNumberType
-        {
-            kCFNumberSInt8Type = 1,
-            kCFNumberSInt16Type = 2,
-            kCFNumberSInt32Type = 3,
-            kCFNumberSInt64Type = 4,
-            kCFNumberFloat32Type = 5,
-            kCFNumberFloat64Type = 6,
-            kCFNumberCharType = 7,
-            kCFNumberShortType = 8,
-            kCFNumberIntType = 9,
-            kCFNumberLongType = 10,
-            kCFNumberLongLongType = 11,
-            kCFNumberFloatType = 12,
-            kCFNumberDoubleType = 13,
-            kCFNumberCFIndexType = 14,
-            kCFNumberNSIntegerType = 15,
-            kCFNumberCGFloatType = 16,
-            kCFNumberMaxType = 16
-        };
-
-
-
-
-
-
         private static byte[] StringToCString(string value)
         {
             byte[] bytes = new byte[value.Length + 1];
@@ -185,18 +521,18 @@ namespace Manzana
 
         public static TypedPtr<CFString> CFStringMakeConstantString(string s)
         {
-            return __CFStringMakeConstantString(StringToCString(s));
+            return CoreImpl.__CFStringMakeConstantString(StringToCString(s));
         }
 
         public static TypedPtr<CFURL> CFURLCreateWithFileSystemPath(IntPtr Allocator, TypedPtr<CFString> FilePath, CFURLPathStyle PathStyle, int isDirectory)
         {
-            return (TypedPtr<CFURL>)(CFURLCreateWithFileSystemPath(Allocator, FilePath.Handle, PathStyle, isDirectory));
+            return (TypedPtr<CFURL>)(CoreImpl.CFURLCreateWithFileSystemPath(Allocator, FilePath.Handle, PathStyle, isDirectory));
         }
 
         public static string CFStringGetCString(TypedPtr<CFString> InString)
         {
             byte[] bytes = new byte[2048];
-            CFStringGetCString(InString.Handle, bytes, 2048, CFStringBuiltInEncodings.kCFStringEncodingUTF8);
+            CoreImpl.CFStringGetCString(InString.Handle, bytes, 2048, CFStringBuiltInEncodings.kCFStringEncodingUTF8);
 
             int ValidLength = 0;
             foreach (byte b in bytes)
@@ -216,7 +552,7 @@ namespace Manzana
 
         public static string GetStringForUrl(TypedPtr<CFURL> Url)
         {
-            IntPtr cfString = CFURLGetString((IntPtr)Url);
+			TypedPtr<CFString> cfString = CoreImpl.CFURLGetString((IntPtr)Url);
             return CFStringGetCString((IntPtr)cfString);
         }
 
@@ -227,12 +563,12 @@ namespace Manzana
 
         public static TypedPtr<CFNumber> CFNumberCreate(Int32 Value)
         {
-            return CFNumberCreate(kCFAllocatorDefault, CFNumberType.kCFNumberSInt32Type, ref Value);
+			return CoreImpl.CFNumberCreate(kCFAllocatorDefault, CFNumberType.kCFNumberSInt32Type, ref Value);
         }
 
         public static TypedPtr<CFDictionary> CreateCFDictionaryHelper()
         {
-            return CFDictionaryCreateMutable(kCFAllocatorDefault, 0, kCFTypeDictionaryKeyCallBacks, kCFTypeDictionaryValueCallBacks);
+			return CoreImpl.CFDictionaryCreateMutable(kCFAllocatorDefault, 0, kCFTypeDictionaryKeyCallBacks, kCFTypeDictionaryValueCallBacks);
         }
 
         /// <summary>
@@ -248,21 +584,21 @@ namespace Manzana
             }
 
             // Convert based on the type of the source object
-            uint TypeID = CFGetTypeID(Source);
+			uint TypeID = CoreImpl.CFGetTypeID(Source);
 
-            if (TypeID == CFNumberGetTypeID())
+			if (TypeID == CoreImpl.CFNumberGetTypeID())
             {
                 return ConvertCFNumber(Source);
             }
-            else if (TypeID == CFStringGetTypeID())
+			else if (TypeID == CoreImpl.CFStringGetTypeID())
             {
                 return CFStringGetCString(Source);
             }
-            else if (TypeID == CFDictionaryGetTypeID())
+			else if (TypeID == CoreImpl.CFDictionaryGetTypeID())
             {
                 return ConvertCFDictionaryToDictionary(Source);
             }
-            else if (TypeID == CFBooleanGetTypeID())
+			else if (TypeID == CoreImpl.CFBooleanGetTypeID())
             {
                 return (Source == kCFBooleanTrue.Handle);
             }
@@ -275,7 +611,7 @@ namespace Manzana
         /// </summary>
         public static double ConvertCFNumber(TypedPtr<CFNumber> Number)
         {
-            CFNumberType TypeID = CFNumberGetType((IntPtr)Number);
+			CFNumberType TypeID = CoreImpl.CFNumberGetType((IntPtr)Number);
 
             switch (TypeID)
             {
@@ -283,7 +619,7 @@ namespace Manzana
                 case CFNumberType.kCFNumberFloatType:
                     {
                         float Result;
-                        CFNumberGetValue((IntPtr)Number, CFNumberType.kCFNumberFloat32Type, out Result);
+						CoreImpl.CFNumberGetValue((IntPtr)Number, CFNumberType.kCFNumberFloat32Type, out Result);
                         return Result;
                     }
 
@@ -292,7 +628,7 @@ namespace Manzana
                 case CFNumberType.kCFNumberCGFloatType:
                     {
                         double Result;
-                        CFNumberGetValue((IntPtr)Number, CFNumberType.kCFNumberFloat64Type, out Result);
+						CoreImpl.CFNumberGetValue((IntPtr)Number, CFNumberType.kCFNumberFloat64Type, out Result);
                         return Result;
                     }
 
@@ -306,7 +642,7 @@ namespace Manzana
                 case CFNumberType.kCFNumberCFIndexType:
                     {
                         int Result;
-                        CFNumberGetValue((IntPtr)Number, CFNumberType.kCFNumberIntType, out Result);
+						CoreImpl.CFNumberGetValue((IntPtr)Number, CFNumberType.kCFNumberIntType, out Result);
                         return Result;
                     }
 
@@ -315,7 +651,7 @@ namespace Manzana
                 case CFNumberType.kCFNumberNSIntegerType:
                     {
                         Int64 Result;
-                        CFNumberGetValue((IntPtr)Number, CFNumberType.kCFNumberSInt64Type, out Result);
+						CoreImpl.CFNumberGetValue((IntPtr)Number, CFNumberType.kCFNumberSInt64Type, out Result);
                         return Result;
                     }
 
@@ -335,7 +671,7 @@ namespace Manzana
             IntPtr [] Keys = new IntPtr[NumPairs];
             IntPtr [] Values = new IntPtr[NumPairs];
 
-            CFDictionaryGetKeysAndValues((IntPtr)Dict, Keys, Values);
+			CoreImpl.CFDictionaryGetKeysAndValues((IntPtr)Dict, Keys, Values);
 
             // Convert each key and value to a managed type
             Dictionary<object, object> Result = new Dictionary<object, object>();
@@ -354,7 +690,7 @@ namespace Manzana
             return Result;
         }
 
-        public static Dictionary<string, object> ConvertCFDictionaryToDictionaryStringy(TypedPtr<CFDictionary> Dict)
+        public static Dictionary<string, object> ConvertCFDictionaryToDictionaryString(TypedPtr<CFDictionary> Dict)
         {
             Dictionary<string, object> Result = new Dictionary<string,object>();
             Dictionary<object, object> Temp = ConvertCFDictionaryToDictionary(Dict);
@@ -370,12 +706,12 @@ namespace Manzana
 
         public static int CFDictionaryGetCount(TypedPtr<CFDictionary> Dict)
         {
-            return CFDictionaryGetCount((IntPtr)Dict);
+			return CoreImpl.CFDictionaryGetCount(Dict);
         }
 
         public static void CFDictionaryAddHelper(TypedPtr<CFDictionary> InDict, string Key, string Value)
         {
-            CFDictionarySetValue((IntPtr)InDict, (IntPtr)CFStringMakeConstantString(Key), (IntPtr)CFStringMakeConstantString(Value));
+			CoreImpl.CFDictionarySetValue((IntPtr)InDict, (IntPtr)CFStringMakeConstantString(Key), (IntPtr)CFStringMakeConstantString(Value));
         }
     }
 }
\ No newline at end of file
diff --git a/Engine/Source/Programs/IOS/MobileDeviceInterface/IPhoneFile.cs b/Engine/Source/Programs/IOS/MobileDeviceInterface/IPhoneFile.cs
--- a/Engine/Source/Programs/IOS/MobileDeviceInterface/IPhoneFile.cs
+++ b/Engine/Source/Programs/IOS/MobileDeviceInterface/IPhoneFile.cs
@@ -126,7 +126,7 @@ namespace Manzana
                 uint ret;
                 ret = 0;
 
-                MobileDevice.AFC.FileRefTell(phone.AFCHandle, handle, ref ret);
+                MobileDevice.DeviceImpl.FileRefTell(phone.AFCHandle, handle, ref ret);
                 return (long)ret;
             }
             set
@@ -141,9 +141,9 @@ namespace Manzana
         /// <param name="value">The new length of the stream.</param>
         public override void SetLength(long value)
         {
-            int ret;
+			//            int ret;
 
-            ret = MobileDevice.AFC.AFCFileRefSetFileSize((IntPtr)phone.AFCHandle, handle, (uint)value);
+			MobileDevice.DeviceImpl.FileRefSetFileSize((IntPtr)phone.AFCHandle, handle, (uint)value);
         }
         #endregion	// Public Properties
 
@@ -158,7 +158,7 @@ namespace Manzana
             {
                 if (handle != 0)
                 {
-                    MobileDevice.AFC.FileRefClose(phone.AFCHandle, handle);
+					MobileDevice.DeviceImpl.FileRefClose(phone.AFCHandle, handle);
                     handle = 0;
                 }
             }
@@ -185,7 +185,7 @@ namespace Manzana
                 temp = new byte[count];
 
             uint len = (uint)count;
-            int ret = MobileDevice.AFC.FileRefRead(phone.AFCHandle, handle, temp, ref len);
+			int ret = MobileDevice.DeviceImpl.FileRefRead(phone.AFCHandle, handle, temp, ref len);
             if (ret != 0)
                 throw new IOException("AFCFileRefRead error = " + ret.ToString());
 
@@ -216,7 +216,7 @@ namespace Manzana
                 Buffer.BlockCopy(buffer, offset, temp, 0, count);
             }
 
-            int ret = MobileDevice.AFC.FileRefWrite(phone.AFCHandle, handle, temp, (uint)count);
+			MobileDevice.DeviceImpl.FileRefWrite(phone.AFCHandle, handle, temp, (uint)count);
         }
 
         /// <summary>
@@ -229,7 +229,7 @@ namespace Manzana
         {
             int ret;
 
-            ret = MobileDevice.AFC.FileRefSeek(phone.AFCHandle, handle, (uint)offset, 0);
+			ret = MobileDevice.DeviceImpl.FileRefSeek(phone.AFCHandle, handle, (uint)offset, 0);
             Console.WriteLine("ret = {0}", ret);
             return offset;
         }
@@ -239,7 +239,7 @@ namespace Manzana
         /// </summary>
         public override void Flush()
         {
-            MobileDevice.AFC.FlushData(phone.AFCHandle, handle);
+			MobileDevice.DeviceImpl.FlushData(phone.AFCHandle, handle);
         }
         #endregion	// Public Methods
 
@@ -267,7 +267,7 @@ namespace Manzana
             }
 
             full_path = phone.FullPath(phone.CurrentDirectory, path);
-            ret = MobileDevice.AFC.FileRefOpen(phone.AFCHandle, full_path, (int)mode, 0, out handle);
+			ret = MobileDevice.DeviceImpl.FileRefOpen(phone.AFCHandle, full_path, (int)mode, 0, out handle);
             if (ret != 0)
             {
                 phone.Reconnect();
diff --git a/Engine/Source/Programs/IOS/MobileDeviceInterface/MobileDevice.cs b/Engine/Source/Programs/IOS/MobileDeviceInterface/MobileDevice.cs
--- a/Engine/Source/Programs/IOS/MobileDeviceInterface/MobileDevice.cs
+++ b/Engine/Source/Programs/IOS/MobileDeviceInterface/MobileDevice.cs
@@ -112,578 +112,1587 @@ namespace Manzana
     [UnmanagedFunctionPointer(CallingConvention.Cdecl)]
     internal delegate void DeviceInstallationCallback(IntPtr/*CFDictionary*/ UntypedStatusDict, IntPtr UserData);
 
-    internal partial class MobileDevice
-    {
-        const string DLLName = "iTunesMobileDevice.dll";
-        static readonly int ForceStaticInit = 42;
+	internal interface MobileDeviceImpl
+	{
+		int NotificationSubscribe(DeviceNotificationCallback DeviceCallbackHandle);
+		IntPtr CopyDeviceIdentifier(TypedPtr<AppleMobileDeviceConnection> device);
+		int Connect(TypedPtr<AppleMobileDeviceConnection> device);
+		IntPtr CopyValue(TypedPtr<AppleMobileDeviceConnection> device, uint unknown, TypedPtr<CFString> cfstring);
+		int Disconnect(TypedPtr<AppleMobileDeviceConnection> device);
+		int GetConnectionID(TypedPtr<AppleMobileDeviceConnection> device);
+		int IsPaired(TypedPtr<AppleMobileDeviceConnection> device);
+		int Pair(TypedPtr<AppleMobileDeviceConnection> device);
+		int Unpair(TypedPtr<AppleMobileDeviceConnection> device);
+		int ValidatePairing(TypedPtr<AppleMobileDeviceConnection> device);
+		int LookupApplications(TypedPtr<AppleMobileDeviceConnection> device, IntPtr options, out Dictionary<string, object> AppBundles);
+		int StartSession(TypedPtr<AppleMobileDeviceConnection> device);
+		int StopSession(TypedPtr<AppleMobileDeviceConnection> device);
+		int StartHouseArrestService(TypedPtr<AppleMobileDeviceConnection> device, TypedPtr<CFString> bundleName, IntPtr unknown1, ref IntPtr handle, int unknown2);
+		int StartService(TypedPtr<AppleMobileDeviceConnection> device, TypedPtr<CFString> serviceName, ref IntPtr handle, IntPtr unknown);
+		int RestoreRegisterForDeviceNotifications(
+			DeviceRestoreNotificationCallback dfu_connect,
+			DeviceRestoreNotificationCallback recovery_connect,
+			DeviceRestoreNotificationCallback dfu_disconnect,
+			DeviceRestoreNotificationCallback recovery_disconnect,
+			uint unknown0,
+			IntPtr user_info);
+		int SecureUninstallApplication(
+			IntPtr serviceConnection,
+			TypedPtr<AppleMobileDeviceConnection> DeviceIfConnIsNull,
+			TypedPtr<CFString> ApplicationIdentifer,
+			TypedPtr<CFDictionary> ClientOptions,
+			DeviceInstallationCallback ProgressCallback,
+			IntPtr UserData);
+		int SecureInstallApplication(
+			IntPtr serviceConnection,
+			TypedPtr<AppleMobileDeviceConnection> DeviceIfConnIsNull,
+			TypedPtr<CFURL> UrlPath,
+			TypedPtr<CFDictionary> ClientOptions,
+			DeviceInstallationCallback ProgressCallback,
+			IntPtr UserData);
+		int SecureUpgradeApplication(
+			IntPtr ServiceConnection,
+			TypedPtr<AppleMobileDeviceConnection> DeviceIfConnIsNull,
+			TypedPtr<CFURL> UrlPath,
+			TypedPtr<CFDictionary> ClientOptions,
+			DeviceInstallationCallback ProgressCallback,
+			IntPtr UserData);
+
+		int ConnectionClose(TypedPtr<AFCCommConnection> conn);
+		int ConnectionInvalidate(TypedPtr<AFCCommConnection> conn);
+		int ConnectionIsValid(TypedPtr<AFCCommConnection> conn);
+		int ConnectionOpen(IntPtr handle, uint io_timeout, out TypedPtr<AFCCommConnection> OutConn);
+		int DeviceInfoOpen(IntPtr handle, ref IntPtr dict);
+		int DirectoryClose(TypedPtr<AFCCommConnection> conn, IntPtr dir);
+		int DirectoryCreate(TypedPtr<AFCCommConnection> conn, string path);
+		int DirectoryOpen(TypedPtr<AFCCommConnection> conn, byte[] path, ref IntPtr dir);
+		int DirectoryRead(TypedPtr<AFCCommConnection> conn, IntPtr dir, ref IntPtr dirent);
+		int FileInfoOpen(TypedPtr<AFCCommConnection> conn, string path, out TypedPtr<AFCDictionary> OutDict);
+		int FileRefClose(TypedPtr<AFCCommConnection> conn, Int64 handle);
+		int FileRefOpen(TypedPtr<AFCCommConnection> conn, string path, int mode, int unknown, out Int64 handle);
+		int FileRefRead(TypedPtr<AFCCommConnection> conn, Int64 handle, byte[] buffer, ref uint len);
+		int FileRefSeek(TypedPtr<AFCCommConnection> conn, Int64 handle, Int64 pos, Int64 origin);
+		int FileRefSetFileSize(TypedPtr<AFCCommConnection> conn, Int64 handle, uint size);
+		int FileRefTell(TypedPtr<AFCCommConnection> conn, Int64 handle, ref uint position);
+		int FileRefWrite(TypedPtr<AFCCommConnection> conn, Int64 handle, byte[] buffer, uint len);
+		int FlushData(TypedPtr<AFCCommConnection> conn, Int64 handle);
+		int KeyValueClose(TypedPtr<AFCDictionary> dict);
+		int KeyValueRead(TypedPtr<AFCDictionary> dict, out IntPtr key, out IntPtr val);
+		int RemovePath(TypedPtr<AFCCommConnection> conn, string path);
+		int RenamePath(TypedPtr<AFCCommConnection> conn, string OldPath, string NewPath);
+		int DirectoryOpen(TypedPtr<AFCCommConnection> conn, string path, ref IntPtr dir);
+		int DirectoryRead(TypedPtr<AFCCommConnection> conn, IntPtr dir, ref string buffer);
+	}
+
+	internal class MobileDeviceOSX : MobileDeviceImpl
+	{
+		//		static readonly int ForceStaticInit = 42;
+		const string DLLName = "/System/Library/PrivateFrameworks/MobileDevice.framework/Versions/A/MobileDevice";
+
+		public int NotificationSubscribe(DeviceNotificationCallback DeviceCallbackHandle)
+		{
+			return AMDeviceMethods.NotificationSubscribe(DeviceCallbackHandle);
+		}
+
+		public IntPtr CopyDeviceIdentifier(TypedPtr<AppleMobileDeviceConnection> device)
+		{
+			return AMDeviceMethods.CopyDeviceIdentifier(device);
+		}
+
+		public int Connect(TypedPtr<AppleMobileDeviceConnection> device)
+		{
+			return AMDeviceMethods.Connect(device);
+		}
+
+		public IntPtr CopyValue(TypedPtr<AppleMobileDeviceConnection> device, uint unknown, TypedPtr<CFString> cfstring)
+		{
+			return AMDeviceMethods.CopyValue(device, unknown, cfstring);
+		}
+
+		public int Disconnect(TypedPtr<AppleMobileDeviceConnection> device)
+		{
+			return AMDeviceMethods.Disconnect(device);
+		}
+
+		public int GetConnectionID(TypedPtr<AppleMobileDeviceConnection> device)
+		{
+			return AMDeviceMethods.GetConnectionID(device);
+		}
+
+		public int IsPaired(TypedPtr<AppleMobileDeviceConnection> device)
+		{
+			return AMDeviceMethods.IsPaired(device);
+		}
+
+		public int Pair(TypedPtr<AppleMobileDeviceConnection> device)
+		{
+			return AMDeviceMethods.Pair(device);
+		}
+
+		public int Unpair(TypedPtr<AppleMobileDeviceConnection> device)
+		{
+			return AMDeviceMethods.Unpair(device);
+		}
+
+		public int ValidatePairing(TypedPtr<AppleMobileDeviceConnection> device)
+		{
+			return AMDeviceMethods.ValidatePairing(device);
+		}
+
+		public int LookupApplications(TypedPtr<AppleMobileDeviceConnection> device, IntPtr options, out Dictionary<string, object> AppBundles)
+		{
+			int Result = AMDeviceMethods.LookupApplications(device, options, out AppBundles);
+			return Result;
+		}
+
+		public int StartSession(TypedPtr<AppleMobileDeviceConnection> device)
+		{
+			return AMDeviceMethods.StartSession((IntPtr)device);
+		}
+
+		public int StopSession(TypedPtr<AppleMobileDeviceConnection> device)
+		{
+			return AMDeviceMethods.StopSession((IntPtr)device);
+		}
+
+		public int StartHouseArrestService(TypedPtr<AppleMobileDeviceConnection> device, TypedPtr<CFString> bundleName, IntPtr unknown1, ref IntPtr handle, int unknown2)
+		{
+			return AMDeviceMethods.StartHouseArrestService((IntPtr)device, (IntPtr)bundleName, unknown1, ref handle, unknown2);
+		}
+
+		public int StartService(TypedPtr<AppleMobileDeviceConnection> device, TypedPtr<CFString> serviceName, ref IntPtr handle, IntPtr unknown)
+		{
+			return AMDeviceMethods.StartService((IntPtr)device, (IntPtr)serviceName, ref handle, unknown);
+		}
+
+		public int RestoreRegisterForDeviceNotifications(
+			DeviceRestoreNotificationCallback dfu_connect,
+			DeviceRestoreNotificationCallback recovery_connect,
+			DeviceRestoreNotificationCallback dfu_disconnect,
+			DeviceRestoreNotificationCallback recovery_disconnect,
+			uint unknown0,
+			IntPtr user_info)
+		{
+			return AMDeviceMethods.AMRestoreRegisterForDeviceNotifications(dfu_connect, recovery_connect, dfu_disconnect, recovery_disconnect, unknown0, user_info);
+		}
+
+		public int SecureUninstallApplication(
+			IntPtr serviceConnection,
+			TypedPtr<AppleMobileDeviceConnection> DeviceIfConnIsNull,
+			TypedPtr<CFString> ApplicationIdentifer,
+			TypedPtr<CFDictionary> ClientOptions,
+			DeviceInstallationCallback ProgressCallback,
+			IntPtr UserData)
+		{
+			return AMDeviceMethods.SecureUninstallApplication(serviceConnection, DeviceIfConnIsNull, ApplicationIdentifer, ClientOptions, ProgressCallback, UserData);
+		}
+
+		public int SecureInstallApplication(
+			IntPtr serviceConnection,
+			TypedPtr<AppleMobileDeviceConnection> DeviceIfConnIsNull,
+			TypedPtr<CFURL> UrlPath,
+			TypedPtr<CFDictionary> ClientOptions,
+			DeviceInstallationCallback ProgressCallback,
+			IntPtr UserData)
+		{
+			return AMDeviceMethods.SecureInstallApplication(serviceConnection, DeviceIfConnIsNull, UrlPath, ClientOptions, ProgressCallback, UserData);
+		}
+
+		public int SecureUpgradeApplication(
+			IntPtr ServiceConnection,
+			TypedPtr<AppleMobileDeviceConnection> DeviceIfConnIsNull,
+			TypedPtr<CFURL> UrlPath,
+			TypedPtr<CFDictionary> ClientOptions,
+			DeviceInstallationCallback ProgressCallback,
+			IntPtr UserData)
+		{
+			return AMDeviceMethods.SecureUpgradeApplication(ServiceConnection, DeviceIfConnIsNull, UrlPath, ClientOptions, ProgressCallback, UserData);
+		}
+
+		public int ConnectionClose(TypedPtr<AFCCommConnection> conn)
+		{
+			return AFC.ConnectionClose(conn);
+		}
+
+		public int ConnectionInvalidate(TypedPtr<AFCCommConnection> conn)
+		{
+			return AFC.ConnectionInvalidate(conn);
+		}
+
+		public int ConnectionIsValid(TypedPtr<AFCCommConnection> conn)
+		{
+			return AFC.ConnectionIsValid(conn);
+		}
+
+		public int ConnectionOpen(IntPtr handle, uint io_timeout, out TypedPtr<AFCCommConnection> OutConn)
+		{
+			return AFC.ConnectionOpen(handle, io_timeout, out OutConn);
+		}
+
+		public int DeviceInfoOpen(IntPtr handle, ref IntPtr dict)
+		{
+			return AFC.DeviceInfoOpen(handle, ref dict);
+		}
+
+		public int DirectoryClose(TypedPtr<AFCCommConnection> conn, IntPtr dir)
+		{
+			return AFC.DirectoryClose(conn, dir);
+		}
+
+		public int DirectoryCreate(TypedPtr<AFCCommConnection> conn, string path)
+		{
+			return AFC.DirectoryCreate(conn, path);
+		}
+
+		public int DirectoryOpen(TypedPtr<AFCCommConnection> conn, byte[] path, ref IntPtr dir)
+		{
+			return AFC.DirectoryOpen(conn, path, ref dir);
+		}
+
+		public int DirectoryRead(TypedPtr<AFCCommConnection> conn, IntPtr dir, ref IntPtr dirent)
+		{
+			return AFC.DirectoryRead(conn, dir, ref dirent);
+		}
+
+		public int FileInfoOpen(TypedPtr<AFCCommConnection> conn, string path, out TypedPtr<AFCDictionary> OutDict)
+		{
+			return AFC.FileInfoOpen(conn, path, out OutDict);
+		}
+
+		public int FileRefClose(TypedPtr<AFCCommConnection> conn, Int64 handle)
+		{
+			return AFC.FileRefClose(conn, handle);
+		}
+
+		public int FileRefOpen(TypedPtr<AFCCommConnection> conn, string path, int mode, int unknown, out Int64 handle)
+		{
+			return AFC.FileRefOpen(conn, path, mode, unknown, out handle);
+		}
+
+		public int FileRefRead(TypedPtr<AFCCommConnection> conn, Int64 handle, byte[] buffer, ref uint len)
+		{
+			return AFC.FileRefRead(conn, handle, buffer, ref len);
+		}
+
+		public int FileRefSeek(TypedPtr<AFCCommConnection> conn, Int64 handle, Int64 pos, Int64 origin)
+		{
+			return AFC.FileRefSeek(conn, handle, pos, origin);
+		}
+
+		public int FileRefSetFileSize(TypedPtr<AFCCommConnection> conn, Int64 handle, uint size)
+		{
+			return AFC.FileRefSetFileSize(conn, handle, size);
+		}
+
+		public int FileRefTell(TypedPtr<AFCCommConnection> conn, Int64 handle, ref uint position)
+		{
+			return AFC.FileRefTell(conn, handle, ref position);
+		}
+
+		public int FileRefWrite(TypedPtr<AFCCommConnection> conn, Int64 handle, byte[] buffer, uint len)
+		{
+			return AFC.FileRefWrite(conn, handle, buffer, len);
+		}
+
+		public int FlushData(TypedPtr<AFCCommConnection> conn, Int64 handle)
+		{
+			return AFC.FlushData(conn, handle);
+		}
+
+		public int KeyValueClose(TypedPtr<AFCDictionary> dict)
+		{
+			return AFC.KeyValueClose(dict);
+		}
+
+		public int KeyValueRead(TypedPtr<AFCDictionary> dict, out IntPtr key, out IntPtr val)
+		{
+			return AFC.KeyValueRead(dict, out key, out val);
+		}
+
+		public int RemovePath(TypedPtr<AFCCommConnection> conn, string path)
+		{
+			return AFC.RemovePath(conn, path);
+		}
+
+		public int RenamePath(TypedPtr<AFCCommConnection> conn, string OldPath, string NewPath)
+		{
+			return AFC.RenamePath(conn, OldPath, NewPath);
+		}
 
-        static MobileDevice()
-        {
-			List<string> PathBits = new List<string>();
-			PathBits.Add(Environment.GetEnvironmentVariable("Path"));
+#region Application Methods
+		public class AMDeviceMethods
+		{
+			// Need to have the path manipulation code in static MobileDevice() to run
+			//			static int ForceMobileDeviceConstructorToRun;
+			static AMDeviceMethods()
+			{
+				//				ForceMobileDeviceConstructorToRun = ForceStaticInit;
+			}
 
-			// Try to add the paths from the registry (they aren't always available on newer iTunes installs though)
-			object RegistryDllPath = Registry.GetValue(@"HKEY_LOCAL_MACHINE\SOFTWARE\Apple Inc.\Apple Mobile Device Support\Shared", "iTunesMobileDeviceDLL", DLLName);
-			if (RegistryDllPath != null)
+			public static IntPtr CopyDeviceIdentifier(TypedPtr<AppleMobileDeviceConnection> device)
 			{
-				FileInfo iTunesMobileDeviceFile = new FileInfo(RegistryDllPath.ToString());
+				return AMDeviceCopyDeviceIdentifier((IntPtr)device);
+			}
 
-				if (iTunesMobileDeviceFile.Exists)
-				{
-					PathBits.Add(iTunesMobileDeviceFile.DirectoryName);
-				}
+			public static int Connect(TypedPtr<AppleMobileDeviceConnection> device)
+			{
+				return AMDeviceConnect((IntPtr)device);
 			}
 
-			object RegistrySupportDir = Registry.GetValue(@"HKEY_LOCAL_MACHINE\SOFTWARE\Apple Inc.\Apple Application Support", "InstallDir", Environment.CurrentDirectory);
-			if (RegistrySupportDir != null)
+			public static IntPtr CopyValue(TypedPtr<AppleMobileDeviceConnection> device, uint unknown, TypedPtr<CFString> cfstring)
 			{
-				DirectoryInfo ApplicationSupportDirectory = new DirectoryInfo(RegistrySupportDir.ToString());
+				return AMDeviceCopyValue((IntPtr)device, unknown, (IntPtr)cfstring);
+			}
 
-				if (ApplicationSupportDirectory.Exists)
-				{
-					PathBits.Add(ApplicationSupportDirectory.FullName);
-				}
+			public static int Disconnect(TypedPtr<AppleMobileDeviceConnection> device)
+			{
+				return AMDeviceDisconnect((IntPtr)device);
 			}
 
-			// Add some guesses as well
-			DirectoryInfo AppleMobileDeviceSupport = new DirectoryInfo(Environment.GetFolderPath(Environment.SpecialFolder.CommonProgramFiles) + @"\Apple\Mobile Device Support");
-			if (AppleMobileDeviceSupport.Exists)
+			public static int GetConnectionID(TypedPtr<AppleMobileDeviceConnection> device)
 			{
-				PathBits.Add(AppleMobileDeviceSupport.FullName);
+				return AMDeviceGetConnectionID((IntPtr)device);
 			}
-	
-			DirectoryInfo AppleMobileDeviceSupportX86 = new DirectoryInfo(Environment.GetFolderPath(Environment.SpecialFolder.CommonProgramFilesX86) + @"\Apple\Mobile Device Support");
-			if ((AppleMobileDeviceSupport != AppleMobileDeviceSupportX86) && (AppleMobileDeviceSupportX86.Exists))
+
+			public static int IsPaired(TypedPtr<AppleMobileDeviceConnection> device)
 			{
-				PathBits.Add(AppleMobileDeviceSupportX86.FullName);
+				return AMDeviceIsPaired((IntPtr)device);
 			}
 
-			// Set the path from all the individual bits
-			Environment.SetEnvironmentVariable("Path", string.Join(";", PathBits));
+			public static int Pair(TypedPtr<AppleMobileDeviceConnection> device)
+			{
+				return AMDevicePair((IntPtr)device);
+			}
 
-            // Initialize the CoreFoundation bindings
-            InitializeCoreFoundation();
-        }
+			public static int Unpair(TypedPtr<AppleMobileDeviceConnection> device)
+			{
+				return AMDeviceUnpair((IntPtr)device);
+			}
 
-#region Application Methods
-        public class AMDeviceMethods
-        {
-            // Need to have the path manipulation code in static MobileDevice() to run
-            static int ForceMobileDeviceConstructorToRun;
-            static AMDeviceMethods()
-            {
-				ForceMobileDeviceConstructorToRun = ForceStaticInit;
-            }
+			public static int ValidatePairing(TypedPtr<AppleMobileDeviceConnection> device)
+			{
+				return AMDeviceValidatePairing((IntPtr)device);
+			}
 
-            public static IntPtr CopyDeviceIdentifier(TypedPtr<AppleMobileDeviceConnection> device)
-            {
-                return AMDeviceCopyDeviceIdentifier((IntPtr)device);
-            }
+			public static int LookupApplications(TypedPtr<AppleMobileDeviceConnection> device, IntPtr options, out Dictionary<string, object> AppBundles)
+			{
+				IntPtr UntypedDict = IntPtr.Zero;
+				int Result = AMDeviceLookupApplications((IntPtr)device, options, ref UntypedDict);
 
-            public static int Connect(TypedPtr<AppleMobileDeviceConnection> device)
-            {
-                return AMDeviceConnect((IntPtr)device);
-            }
+				AppBundles = MobileDevice.ConvertCFDictionaryToDictionaryString(new TypedPtr<CFDictionary>(UntypedDict));
 
-            public static IntPtr CopyValue(TypedPtr<AppleMobileDeviceConnection> device, uint unknown, TypedPtr<CFString> cfstring)
-            {
-                return AMDeviceCopyValue((IntPtr)device, unknown, (IntPtr)cfstring);
-            }
+				return Result;
+			}
 
-            public static int Disconnect(TypedPtr<AppleMobileDeviceConnection> device)
-            {
-                return AMDeviceDisconnect((IntPtr)device);
-            }
+			public static int StartSession(TypedPtr<AppleMobileDeviceConnection> device)
+			{
+				return AMDeviceStartSession((IntPtr)device);
+			}
 
-            public static int GetConnectionID(TypedPtr<AppleMobileDeviceConnection> device)
-            {
-                return AMDeviceGetConnectionID((IntPtr)device);
-            }
+			public static int StopSession(TypedPtr<AppleMobileDeviceConnection> device)
+			{
+				return AMDeviceStopSession((IntPtr)device);
+			}
 
-            public static int IsPaired(TypedPtr<AppleMobileDeviceConnection> device)
-            {
-                return AMDeviceIsPaired((IntPtr)device);
-            }
+			public static int StartHouseArrestService(TypedPtr<AppleMobileDeviceConnection> device, TypedPtr<CFString> bundleName, IntPtr unknown1, ref IntPtr handle, int unknown2)
+			{
+				return AMDeviceStartHouseArrestService((IntPtr)device, (IntPtr)bundleName, unknown1, ref handle, unknown2);
+			}
 
-            public static int Pair(TypedPtr<AppleMobileDeviceConnection> device)
-            {
-                return AMDevicePair((IntPtr)device);
-            }
+			public static int StartService(TypedPtr<AppleMobileDeviceConnection> device, TypedPtr<CFString> serviceName, ref IntPtr handle, IntPtr unknown)
+			{
+				return AMDeviceStartService((IntPtr)device, (IntPtr)serviceName, ref handle, unknown);
+			}
 
-            public static int Unpair(TypedPtr<AppleMobileDeviceConnection> device)
-            {
-                return AMDeviceUnpair((IntPtr)device);
-            }
+			public static int UninstallApplication(
+				TypedPtr<AppleMobileDeviceConnection> device,
+				TypedPtr<CFString> ApplicationIdentifier,
+				TypedPtr<CFDictionary> ClientOptions,
+				DeviceInstallationCallback ProgressCallback,
+				IntPtr UserData)
+			{
+				return AMDeviceUninstallApplication((IntPtr)device, (IntPtr)ApplicationIdentifier, (IntPtr)ClientOptions, ProgressCallback, UserData);
+			}
 
-            public static int ValidatePairing(TypedPtr<AppleMobileDeviceConnection> device)
-            {
-                return AMDeviceValidatePairing((IntPtr)device);
-            }
+			public static int SecureUninstallApplication(
+				IntPtr serviceConnection,
+				TypedPtr<AppleMobileDeviceConnection> DeviceIfConnIsNull,
+				TypedPtr<CFString> ApplicationIdentifer,
+				TypedPtr<CFDictionary> ClientOptions,
+				DeviceInstallationCallback ProgressCallback,
+				IntPtr UserData)
+			{
+				return AMDeviceSecureUninstallApplication(serviceConnection, (IntPtr)DeviceIfConnIsNull, (IntPtr)ApplicationIdentifer, (IntPtr)ClientOptions, ProgressCallback, UserData);
+			}
 
-            public static int LookupApplications(TypedPtr<AppleMobileDeviceConnection> device, IntPtr options, out Dictionary<string, object> AppBundles)
-            {
-				IntPtr UntypedDict = IntPtr.Zero;
-				int Result = AMDeviceLookupApplications((IntPtr)device, options, ref UntypedDict);
+			public static int SecureInstallApplication(
+				IntPtr serviceConnection,
+				TypedPtr<AppleMobileDeviceConnection> DeviceIfConnIsNull,
+				TypedPtr<CFURL> UrlPath,
+				TypedPtr<CFDictionary> ClientOptions,
+				DeviceInstallationCallback ProgressCallback,
+				IntPtr UserData)
+			{
+				return AMDeviceSecureInstallApplication(serviceConnection, (IntPtr)DeviceIfConnIsNull, (IntPtr)UrlPath, (IntPtr)ClientOptions, ProgressCallback, UserData);
+			}
+
+			public static int TransferApplication(
+				TypedPtr<AppleMobileDeviceConnection> device,
+				TypedPtr<CFString> InPackagePath,
+				IntPtr UnknownButUnused,
+				DeviceInstallationCallback ProgressCallback,
+				IntPtr UserData)
+			{
+				return AMDeviceTransferApplication((IntPtr)device, (IntPtr)InPackagePath, UnknownButUnused, ProgressCallback, UserData);
+			}
 
-				AppBundles = MobileDevice.ConvertCFDictionaryToDictionaryStringy(new TypedPtr<CFDictionary>(UntypedDict));
+			public static int SecureUpgradeApplication(
+				IntPtr ServiceConnection,
+				TypedPtr<AppleMobileDeviceConnection> DeviceIfConnIsNull,
+				TypedPtr<CFURL> UrlPath,
+				TypedPtr<CFDictionary> ClientOptions,
+				DeviceInstallationCallback ProgressCallback,
+				IntPtr UserData)
+			{
+				return AMDeviceSecureUpgradeApplication(ServiceConnection, (IntPtr)DeviceIfConnIsNull, (IntPtr)UrlPath, (IntPtr)ClientOptions, ProgressCallback, UserData);
+			}
 
+			public static int InstallApplication(
+				TypedPtr<AppleMobileDeviceConnection> device,
+				TypedPtr<CFString> FilePath,
+				TypedPtr<CFDictionary> ClientOptions,
+				DeviceInstallationCallback ProgressCallback,
+				IntPtr UserData)
+			{
+				return AMDeviceInstallApplication((IntPtr)device, (IntPtr)FilePath, (IntPtr)ClientOptions, ProgressCallback, UserData);
+			}
+
+			public static int NotificationSubscribe(DeviceNotificationCallback DeviceCallbackHandle)
+			{
+				IntPtr notification;
+				return AMDeviceNotificationSubscribe(DeviceCallbackHandle, 0, 0, 0, out notification);
+			}
+
+			[DllImport(DLLName, CallingConvention = CallingConvention.Cdecl)]
+			private extern static IntPtr AMDeviceCopyDeviceIdentifier(IntPtr/*AppleMobileDeviceConnection*/ device);
+
+			[DllImport(DLLName, CallingConvention = CallingConvention.Cdecl)]
+			private extern static int AMDeviceConnect(IntPtr/*AppleMobileDeviceConnection*/ device);
+
+			[DllImport(DLLName, CallingConvention = CallingConvention.Cdecl)]
+			private extern static IntPtr AMDeviceCopyValue(IntPtr/*AppleMobileDeviceConnection*/ device, uint unknown, IntPtr/*CFString*/ cfstring);
+
+			[DllImport(DLLName, CallingConvention = CallingConvention.Cdecl)]
+			private extern static int AMDeviceDisconnect(IntPtr/*AppleMobileDeviceConnection*/ device);
+
+			[DllImport(DLLName, CallingConvention = CallingConvention.Cdecl)]
+			private extern static int AMDeviceGetConnectionID(IntPtr/*AppleMobileDeviceConnection*/ device);
+
+			[DllImport(DLLName, CallingConvention = CallingConvention.Cdecl)]
+			private extern static int AMDeviceIsPaired(IntPtr/*AppleMobileDeviceConnection*/ device);
+
+			[DllImport(DLLName, CallingConvention = CallingConvention.Cdecl)]
+			private extern static int AMDevicePair(IntPtr/*AppleMobileDeviceConnection*/ device);
+
+			[DllImport(DLLName, CallingConvention = CallingConvention.Cdecl)]
+			private extern static int AMDeviceUnpair(IntPtr/*AppleMobileDeviceConnection*/ device);
+
+			[DllImport(DLLName, CallingConvention = CallingConvention.Cdecl)]
+			private extern static int AMDeviceValidatePairing(IntPtr/*AppleMobileDeviceConnection*/ device);
+
+			[DllImport(DLLName, CallingConvention = CallingConvention.Cdecl)]
+			private extern static int AMDeviceLookupApplications(IntPtr/*AppleMobileDeviceConnection*/ device, IntPtr options, ref IntPtr appBundles);
+
+			[DllImport(DLLName, CallingConvention = CallingConvention.Cdecl)]
+			private extern static int AMDeviceNotificationSubscribe(DeviceNotificationCallback callback, uint unused1, uint unused2, uint unused3, out IntPtr am_device_notification_ptr);
+
+			[DllImport(DLLName, CallingConvention = CallingConvention.Cdecl)]
+			private extern static int AMDeviceStartHouseArrestService(IntPtr/*AppleMobileDeviceConnection*/ device, IntPtr/*CFString*/ bundleName, IntPtr unknown1, ref IntPtr handle, int unknown2);
+
+			[DllImport(DLLName, CallingConvention = CallingConvention.Cdecl)]
+			private extern static int AMDeviceStartService(IntPtr/*AppleMobileDeviceConnection*/ device, IntPtr/*CFString*/ serviceName, ref IntPtr handle, IntPtr unknown);
+
+			[DllImport(DLLName, CallingConvention = CallingConvention.Cdecl)]
+			private extern static int AMDeviceStartSession(IntPtr/*AppleMobileDeviceConnection*/ device);
+
+			[DllImport(DLLName, CallingConvention = CallingConvention.Cdecl)]
+			private extern static int AMDeviceStopSession(IntPtr/*AppleMobileDeviceConnection*/ device);
+
+			[DllImport(DLLName, CallingConvention = CallingConvention.Cdecl)]
+			private extern static int AMRestoreModeDeviceCreate(uint unknown0, int connection_id, uint unknown1);
+
+			[DllImport(DLLName, CallingConvention = CallingConvention.Cdecl)]
+			public extern static int AMRestoreRegisterForDeviceNotifications(
+				DeviceRestoreNotificationCallback dfu_connect,
+				DeviceRestoreNotificationCallback recovery_connect,
+				DeviceRestoreNotificationCallback dfu_disconnect,
+				DeviceRestoreNotificationCallback recovery_disconnect,
+				uint unknown0,
+				IntPtr user_info);
+
+			[DllImport(DLLName, CallingConvention = CallingConvention.Cdecl)]
+			private extern static int AMDeviceUninstallApplication(
+				IntPtr/*AppleMobileDeviceConnection*/ device,
+				IntPtr/*CFString*/ ApplicationIdentifier,
+				IntPtr/*CFDictionary*/ ClientOptions,
+				DeviceInstallationCallback ProgressCallback,
+				IntPtr UserData);
+
+			[DllImport(DLLName, CallingConvention = CallingConvention.Cdecl)]
+			private extern static int AMDeviceSecureUninstallApplication(
+				IntPtr serviceConnection,
+				IntPtr/*AppleMobileDeviceConnection*/ DeviceIfConnIsNull,
+				IntPtr/*CFString*/ ApplicationIdentifer,
+				IntPtr/*CFDictionary*/ ClientOptions,
+				DeviceInstallationCallback ProgressCallback,
+				IntPtr UserData);
+
+			[DllImport(DLLName, CallingConvention = CallingConvention.Cdecl)]
+			private extern static int AMDeviceSecureInstallApplication(
+				IntPtr serviceConnection,
+				IntPtr/*AppleMobileDeviceConnection*/ DeviceIfConnIsNull,
+				IntPtr/*CFURL*/ UrlPath,
+				IntPtr/*CFDictionary*/ ClientOptions,
+				DeviceInstallationCallback ProgressCallback,
+				IntPtr UserData);
+
+			[DllImport(DLLName, CallingConvention = CallingConvention.Cdecl)]
+			private extern static int AMDeviceTransferApplication(
+				IntPtr/*AppleMobileDeviceConnection*/ device,
+				IntPtr/*CFString*/ InPackagePath,
+				IntPtr UnknownButUnused,
+				DeviceInstallationCallback ProgressCallback,
+				IntPtr UserData);
+
+			[DllImport(DLLName, CallingConvention = CallingConvention.Cdecl)]
+			private extern static int AMDeviceSecureUpgradeApplication(
+				IntPtr ServiceConnection,
+				IntPtr/*AppleMobileDeviceConnection*/ DeviceIfConnIsNull,
+				IntPtr/*CFURL*/ UrlPath,
+				IntPtr/*CFDictionary*/ ClientOptions,
+				DeviceInstallationCallback ProgressCallback,
+				IntPtr UserData);
+
+			[DllImport(DLLName, CallingConvention = CallingConvention.Cdecl)]
+			private extern static int AMDeviceInstallApplication(
+				IntPtr/*AppleMobileDeviceConnection*/ device,
+				IntPtr/*CFString*/ FilePath,
+				IntPtr/*CFDictionary*/ ClientOptions,
+				DeviceInstallationCallback ProgressCallback,
+				IntPtr UserData);
+
+			// the unknown goes into the 4th dword in the connection struct.  if non-zero, winsock send() won't get called in AMDServiceConnectionSend AFAIK
+			// the only option is "CloseOnInvalidate" (CFBoolean type, defaults to false)
+			[DllImport(DLLName, CallingConvention = CallingConvention.Cdecl)]
+			private extern static IntPtr AMDServiceConnectionCreate(IntPtr/*AppleMobileDeviceConnection*/ device, IntPtr UnknownTypically0, IntPtr OptionsDictionary);
+
+
+			[DllImport(DLLName, CallingConvention = CallingConvention.Cdecl)]
+			private extern static int AMSBeginSync(IntPtr/*AppleMobileDeviceConnection*/ device, IntPtr OuterIn8, IntPtr OuterIn12, IntPtr OuterIn16);
+		}
+		#endregion
+
+#region AFC Operations
+
+		public int DirectoryOpen(TypedPtr<AFCCommConnection> conn, string path, ref IntPtr dir)
+		{
+			return AFC.DirectoryOpen((IntPtr)conn, Encoding.UTF8.GetBytes(path), ref dir);
+		}
+
+		public int DirectoryRead(TypedPtr<AFCCommConnection> conn, IntPtr dir, ref string buffer)
+		{
+			int ret;
+
+			IntPtr ptr = IntPtr.Zero;
+			ret = AFC.DirectoryRead((IntPtr)conn, dir, ref ptr);
+			if ((ret == 0) && (ptr != IntPtr.Zero))
+			{
+				buffer = Marshal.PtrToStringAnsi(ptr);
+			}
+			else
+			{
+				buffer = null;
+			}
+			return ret;
+		}
+
+		public class AFC
+		{
+			public static int ConnectionClose(TypedPtr<AFCCommConnection> conn)
+			{
+				return AFCConnectionClose((IntPtr)conn);
+			}
+
+			public static int ConnectionInvalidate(TypedPtr<AFCCommConnection> conn)
+			{
+				return AFCConnectionInvalidate((IntPtr)conn);
+			}
+
+			public static int ConnectionIsValid(TypedPtr<AFCCommConnection> conn)
+			{
+				return AFCConnectionIsValid((IntPtr)conn);
+			}
+
+			public static int ConnectionOpen(IntPtr handle, uint io_timeout, out TypedPtr<AFCCommConnection> OutConn)
+			{
+				IntPtr Conn;
+				int Result = AFCConnectionOpen(handle, io_timeout, out Conn);
+				OutConn = Conn;
 				return Result;
-            }
+			}
 
-            public static int StartSession(TypedPtr<AppleMobileDeviceConnection> device)
-            {
-                return AMDeviceStartSession((IntPtr)device);
-            }
+			public static int DeviceInfoOpen(IntPtr handle, ref IntPtr dict)
+			{
+				return AFCDeviceInfoOpen(handle, ref dict);
+			}
 
-            public static int StopSession(TypedPtr<AppleMobileDeviceConnection> device)
-            {
-                return AMDeviceStopSession((IntPtr)device);
-            }
+			public static int DirectoryClose(TypedPtr<AFCCommConnection> conn, IntPtr dir)
+			{
+				return AFCDirectoryClose((IntPtr)conn, dir);
+			}
 
-            public static int StartHouseArrestService(TypedPtr<AppleMobileDeviceConnection> device, TypedPtr<CFString> bundleName, IntPtr unknown1, ref IntPtr handle, int unknown2)
-            {
-                return AMDeviceStartHouseArrestService((IntPtr)device, (IntPtr)bundleName, unknown1, ref handle, unknown2);
-            }
+			public static int DirectoryCreate(TypedPtr<AFCCommConnection> conn, string path)
+			{
+				return AFCDirectoryCreate((IntPtr)conn, path);
+			}
 
-            public static int StartService(TypedPtr<AppleMobileDeviceConnection> device, TypedPtr<CFString> serviceName, ref IntPtr handle, IntPtr unknown)
-            {
-                return AMDeviceStartService((IntPtr)device, (IntPtr)serviceName, ref handle, unknown);
-            }
+			public static int DirectoryOpen(TypedPtr<AFCCommConnection> conn, byte[] path, ref IntPtr dir)
+			{
+				return AFCDirectoryOpen((IntPtr)conn, path, ref dir);
+			}
 
-            
-            
-            
-            
-            
-            
-            
-            public static int UninstallApplication(
-                TypedPtr<AppleMobileDeviceConnection> device,
-                TypedPtr<CFString> ApplicationIdentifier,
-                TypedPtr<CFDictionary> ClientOptions,
-                DeviceInstallationCallback ProgressCallback,
-                IntPtr UserData)
-            {
-                return AMDeviceUninstallApplication((IntPtr)device, (IntPtr)ApplicationIdentifier, (IntPtr)ClientOptions, ProgressCallback, UserData); 
-            }
+			public static int DirectoryRead(TypedPtr<AFCCommConnection> conn, IntPtr dir, ref IntPtr dirent)
+			{
+				return AFCDirectoryRead((IntPtr)conn, dir, ref dirent);
+			}
 
-            public static int SecureUninstallApplication(
-                IntPtr serviceConnection,
-                TypedPtr<AppleMobileDeviceConnection> DeviceIfConnIsNull,
-                TypedPtr<CFString> ApplicationIdentifer,
-                TypedPtr<CFDictionary> ClientOptions,
-                DeviceInstallationCallback ProgressCallback,
-                IntPtr UserData)
-            {
-                return AMDeviceSecureUninstallApplication(serviceConnection, (IntPtr)DeviceIfConnIsNull, (IntPtr)ApplicationIdentifer, (IntPtr)ClientOptions, ProgressCallback, UserData);
-            }
+			public static int FileInfoOpen(TypedPtr<AFCCommConnection> conn, string path, out TypedPtr<AFCDictionary> OutDict)
+			{
+				IntPtr UntypedDict;
+				int Result = AFCFileInfoOpen((IntPtr)conn, path, out UntypedDict);
+				OutDict = UntypedDict;
 
-            public static int SecureInstallApplication(
-                IntPtr serviceConnection,
-                TypedPtr<AppleMobileDeviceConnection> DeviceIfConnIsNull,
-                TypedPtr<CFURL> UrlPath,
-                TypedPtr<CFDictionary> ClientOptions,
-                DeviceInstallationCallback ProgressCallback,
-                IntPtr UserData)
-            {
-                return AMDeviceSecureInstallApplication(serviceConnection, (IntPtr)DeviceIfConnIsNull, (IntPtr)UrlPath, (IntPtr)ClientOptions, ProgressCallback, UserData);
-            }
+				return Result;
+			}
 
-            public static int TransferApplication(
-                TypedPtr<AppleMobileDeviceConnection> device,
-                TypedPtr<CFString> InPackagePath,
-                IntPtr UnknownButUnused,
-                DeviceInstallationCallback ProgressCallback,
-                IntPtr UserData)
-            {
-                return AMDeviceTransferApplication((IntPtr)device, (IntPtr)InPackagePath, UnknownButUnused, ProgressCallback, UserData);
-            }
+			public static int FileRefClose(TypedPtr<AFCCommConnection> conn, Int64 handle)
+			{
+				return AFCFileRefClose((IntPtr)conn, handle);
+			}
 
-            public static int SecureUpgradeApplication(
-                IntPtr ServiceConnection,
-                TypedPtr<AppleMobileDeviceConnection> DeviceIfConnIsNull,
-                TypedPtr<CFURL> UrlPath,
-                TypedPtr<CFDictionary> ClientOptions,
-                DeviceInstallationCallback ProgressCallback,
-                IntPtr UserData)
-            {
-                return AMDeviceSecureUpgradeApplication(ServiceConnection, (IntPtr)DeviceIfConnIsNull, (IntPtr)UrlPath, (IntPtr)ClientOptions, ProgressCallback, UserData);
-            }
+			public static int FileRefOpen(TypedPtr<AFCCommConnection> conn, string path, int mode, int unknown, out Int64 handle)
+			{
+				return AFCFileRefOpen((IntPtr)conn, path, mode, unknown, out handle);
+			}
 
-            public static int InstallApplication(
-                TypedPtr<AppleMobileDeviceConnection> device,
-                TypedPtr<CFString> FilePath,
-                TypedPtr<CFDictionary> ClientOptions,
-                DeviceInstallationCallback ProgressCallback,
-                IntPtr UserData)
-            {
-                return AMDeviceInstallApplication((IntPtr)device, (IntPtr)FilePath, (IntPtr)ClientOptions, ProgressCallback, UserData);
-            }
+			public static int FileRefRead(TypedPtr<AFCCommConnection> conn, Int64 handle, byte[] buffer, ref uint len)
+			{
+				return AFCFileRefRead((IntPtr)conn, handle, buffer, ref len);
+			}
 
+			public static int FileRefSetFileSize(TypedPtr<AFCCommConnection> conn, Int64 handle, uint size)
+			{
+				return AFCFileRefSetFileSize((IntPtr)conn, handle, size);
+			}
 
+			public static int FileRefSeek(TypedPtr<AFCCommConnection> conn, Int64 handle, Int64 pos, Int64 origin)
+			{
+				return AFCFileRefSeek((IntPtr)conn, handle, pos, origin);
+			}
 
-            public static int NotificationSubscribe(DeviceNotificationCallback DeviceCallbackHandle)
-            {
-                IntPtr notification;
-                return AMDeviceNotificationSubscribe(DeviceCallbackHandle, 0, 0, 0, out notification);
-            }
+			public static int FileRefTell(TypedPtr<AFCCommConnection> conn, Int64 handle, ref uint position)
+			{
+				return AFCFileRefTell((IntPtr)conn, handle, ref position);
+			}
 
+			public static int FileRefWrite(TypedPtr<AFCCommConnection> conn, Int64 handle, byte[] buffer, uint len)
+			{
+				return AFCFileRefWrite((IntPtr)conn, handle, buffer, len);
+			}
+
+			public static int FlushData(TypedPtr<AFCCommConnection> conn, Int64 handle)
+			{
+				return AFCFlushData((IntPtr)conn, handle);
+			}
+
+			public static int KeyValueClose(TypedPtr<AFCDictionary> dict)
+			{
+				return AFCKeyValueClose((IntPtr)dict);
+			}
+
+			public static int KeyValueRead(TypedPtr<AFCDictionary> dict, out IntPtr key, out IntPtr val)
+			{
+				return AFCKeyValueRead((IntPtr)dict, out key, out val);
+			}
+
+			public static int RemovePath(TypedPtr<AFCCommConnection> conn, string path)
+			{
+				return AFCRemovePath((IntPtr)conn, path);
+			}
+
+			public static int RenamePath(TypedPtr<AFCCommConnection> conn, string OldPath, string NewPath)
+			{
+				return AFCRenamePath((IntPtr)conn, OldPath, NewPath);
+			}
+
+			[DllImport(DLLName, CallingConvention = CallingConvention.Cdecl)]
+			private extern static int AFCConnectionClose(IntPtr/*AFCCommConnection*/ conn);
 
-            
-            
-            [DllImport(DLLName, CallingConvention = CallingConvention.Cdecl)]
-            private extern static IntPtr AMDeviceCopyDeviceIdentifier(IntPtr/*AppleMobileDeviceConnection*/ device);
+			[DllImport(DLLName, CallingConvention = CallingConvention.Cdecl)]
+			private extern static int AFCConnectionInvalidate(IntPtr/*AFCCommConnection*/ conn);
 
-            [DllImport(DLLName, CallingConvention = CallingConvention.Cdecl)]
-            private extern static int AMDeviceConnect(IntPtr/*AppleMobileDeviceConnection*/ device);
+			[DllImport(DLLName, CallingConvention = CallingConvention.Cdecl)]
+			private extern static int AFCConnectionIsValid(IntPtr/*AFCCommConnection*/ conn);
 
-            [DllImport(DLLName, CallingConvention = CallingConvention.Cdecl)]
-            private extern static IntPtr AMDeviceCopyValue(IntPtr/*AppleMobileDeviceConnection*/ device, uint unknown, IntPtr/*CFString*/ cfstring);
+			[DllImport(DLLName, CallingConvention = CallingConvention.Cdecl)]
+			private extern static int AFCConnectionOpen(IntPtr handle, uint io_timeout, out IntPtr/*AFCCommConnection*/ conn);
 
-            [DllImport(DLLName, CallingConvention = CallingConvention.Cdecl)]
-            private extern static int AMDeviceDisconnect(IntPtr/*AppleMobileDeviceConnection*/ device);
+			[DllImport(DLLName, CallingConvention = CallingConvention.Cdecl)]
+			private extern static int AFCDeviceInfoOpen(IntPtr handle, ref IntPtr dict);
 
-            [DllImport(DLLName, CallingConvention = CallingConvention.Cdecl)]
-            private extern static int AMDeviceGetConnectionID(IntPtr/*AppleMobileDeviceConnection*/ device);
+			[DllImport(DLLName, CallingConvention = CallingConvention.Cdecl)]
+			private extern static int AFCDirectoryClose(IntPtr/*AFCCommConnection*/ conn, IntPtr dir);
 
+			[DllImport(DLLName, CallingConvention = CallingConvention.Cdecl)]
+			private extern static int AFCDirectoryCreate(IntPtr/*AFCCommConnection*/ conn, string path);
 
+			[DllImport(DLLName, CallingConvention = CallingConvention.Cdecl)]
+			private extern static int AFCDirectoryOpen(IntPtr/*AFCCommConnection*/ conn, byte[] path, ref IntPtr dir);
 
-            [DllImport(DLLName, CallingConvention = CallingConvention.Cdecl)]
-            private extern static int AMDeviceIsPaired(IntPtr/*AppleMobileDeviceConnection*/ device);
+			[DllImport(DLLName, CallingConvention = CallingConvention.Cdecl)]
+			private extern static int AFCDirectoryRead(IntPtr/*AFCCommConnection*/ conn, IntPtr dir, ref IntPtr dirent);
 
-            [DllImport(DLLName, CallingConvention = CallingConvention.Cdecl)]
-            private extern static int AMDevicePair(IntPtr/*AppleMobileDeviceConnection*/ device);
+			[DllImport(DLLName, CallingConvention = CallingConvention.Cdecl)]
+			private extern static int AFCFileInfoOpen(IntPtr/*AFCCommConnection*/ conn, string path, out IntPtr OutDict);
 
-            [DllImport(DLLName, CallingConvention = CallingConvention.Cdecl)]
-            private extern static int AMDeviceUnpair(IntPtr/*AppleMobileDeviceConnection*/ device);
+			[DllImport(DLLName, CallingConvention = CallingConvention.Cdecl)]
+			private extern static int AFCFileRefClose(IntPtr/*AFCCommConnection*/ conn, Int64 handle);
 
-            [DllImport(DLLName, CallingConvention = CallingConvention.Cdecl)]
-            private extern static int AMDeviceValidatePairing(IntPtr/*AppleMobileDeviceConnection*/ device);
+			[DllImport(DLLName, CallingConvention = CallingConvention.Cdecl)]
+			private extern static int AFCFileRefOpen(IntPtr/*AFCCommConnection*/ conn, string path, int mode, int unknown, out Int64 handle);
 
+			[DllImport(DLLName, CallingConvention = CallingConvention.Cdecl)]
+			private extern static int AFCFileRefRead(IntPtr/*AFCCommConnection*/ conn, Int64 handle, byte[] buffer, ref uint len);
 
-            [DllImport(DLLName, CallingConvention = CallingConvention.Cdecl)]
-            private extern static int AMDeviceLookupApplications(IntPtr/*AppleMobileDeviceConnection*/ device, IntPtr options, ref IntPtr appBundles);
+			[DllImport(DLLName, CallingConvention = CallingConvention.Cdecl)]
+			private extern static int AFCFileRefSeek(IntPtr/*AFCCommConnection*/ conn, Int64 handle, Int64 pos, Int64 origin);
 
+			// FIXME - not working, arguments?
+			[DllImport(DLLName, CallingConvention = CallingConvention.Cdecl)]
+			public extern static int AFCFileRefSetFileSize(IntPtr/*AFCCommConnection*/ conn, Int64 handle, uint size);
 
+			[DllImport(DLLName, CallingConvention = CallingConvention.Cdecl)]
+			private extern static int AFCFileRefTell(IntPtr/*AFCCommConnection*/ conn, Int64 handle, ref uint position);
 
-            [DllImport(DLLName, CallingConvention = CallingConvention.Cdecl)]
-            private extern static int AMDeviceNotificationSubscribe(DeviceNotificationCallback callback, uint unused1, uint unused2, uint unused3, out IntPtr am_device_notification_ptr);
+			[DllImport(DLLName, CallingConvention = CallingConvention.Cdecl)]
+			private extern static int AFCFileRefWrite(IntPtr/*AFCCommConnection*/ conn, Int64 handle, byte[] buffer, uint len);
 
+			[DllImport(DLLName, CallingConvention = CallingConvention.Cdecl)]
+			private extern static int AFCFlushData(IntPtr/*AFCCommConnection*/ conn, Int64 handle);
 
-            [DllImport(DLLName, CallingConvention = CallingConvention.Cdecl)]
-            private extern static int AMDeviceStartHouseArrestService(IntPtr/*AppleMobileDeviceConnection*/ device, IntPtr/*CFString*/ bundleName, IntPtr unknown1, ref IntPtr handle, int unknown2);
+			[DllImport(DLLName, CallingConvention = CallingConvention.Cdecl)]
+			private extern static int AFCKeyValueClose(IntPtr dict);
 
-            [DllImport(DLLName, CallingConvention = CallingConvention.Cdecl)]
-            private extern static int AMDeviceStartService(IntPtr/*AppleMobileDeviceConnection*/ device, IntPtr/*CFString*/ serviceName, ref IntPtr handle, IntPtr unknown);
+			[DllImport(DLLName, CallingConvention = CallingConvention.Cdecl)]
+			private extern static int AFCKeyValueRead(IntPtr dict, out IntPtr key, out IntPtr val);
 
 
-            [DllImport(DLLName, CallingConvention = CallingConvention.Cdecl)]
-            private extern static int AMDeviceStartSession(IntPtr/*AppleMobileDeviceConnection*/ device);
+			[DllImport(DLLName, CallingConvention = CallingConvention.Cdecl)]
+			private extern static int AFCRemovePath(IntPtr/*AFCCommConnection*/ conn, string path);
+
+			[DllImport(DLLName, CallingConvention = CallingConvention.Cdecl)]
+			private extern static int AFCRenamePath(IntPtr/*AFCCommConnection*/ conn, string old_path, string new_path);
+		}
+
+		#endregion
+	}
+
+	internal class MobileDeviceWin : MobileDeviceImpl
+	{
+		//		static readonly int ForceStaticInit = 42;
+		const string DLLName = "iTunesMobileDevice.dll";
+
+		static MobileDeviceWin()
+		{
+			List<string> PathBits = new List<string>();
+			PathBits.Add(Environment.GetEnvironmentVariable("Path"));
 
-            [DllImport(DLLName, CallingConvention = CallingConvention.Cdecl)]
-            private extern static int AMDeviceStopSession(IntPtr/*AppleMobileDeviceConnection*/ device);
+			// Try to add the paths from the registry (they aren't always available on newer iTunes installs though)
+			object RegistryDllPath = Registry.GetValue(@"HKEY_LOCAL_MACHINE\SOFTWARE\Apple Inc.\Apple Mobile Device Support\Shared", "iTunesMobileDeviceDLL", DLLName);
+			if (RegistryDllPath != null)
+			{
+				FileInfo iTunesMobileDeviceFile = new FileInfo(RegistryDllPath.ToString());
 
+				if (iTunesMobileDeviceFile.Exists)
+				{
+					PathBits.Add(iTunesMobileDeviceFile.DirectoryName);
+				}
+			}
 
+			object RegistrySupportDir = Registry.GetValue(@"HKEY_LOCAL_MACHINE\SOFTWARE\Apple Inc.\Apple Application Support", "InstallDir", Environment.CurrentDirectory);
+			if (RegistrySupportDir != null)
+			{
+				DirectoryInfo ApplicationSupportDirectory = new DirectoryInfo(RegistrySupportDir.ToString());
 
+				if (ApplicationSupportDirectory.Exists)
+				{
+					PathBits.Add(ApplicationSupportDirectory.FullName);
+				}
+			}
 
-            [DllImport(DLLName, CallingConvention = CallingConvention.Cdecl)]
-            private extern static int AMRestoreModeDeviceCreate(uint unknown0, int connection_id, uint unknown1);
+			// Add some guesses as well
+			DirectoryInfo AppleMobileDeviceSupport = new DirectoryInfo(Environment.GetFolderPath(Environment.SpecialFolder.CommonProgramFiles) + @"\Apple\Mobile Device Support");
+			if (AppleMobileDeviceSupport.Exists)
+			{
+				PathBits.Add(AppleMobileDeviceSupport.FullName);
+			}
 
-            [DllImport(DLLName, CallingConvention = CallingConvention.Cdecl)]
-            public extern static int AMRestoreRegisterForDeviceNotifications(
-                DeviceRestoreNotificationCallback dfu_connect,
-                DeviceRestoreNotificationCallback recovery_connect,
-                DeviceRestoreNotificationCallback dfu_disconnect,
-                DeviceRestoreNotificationCallback recovery_disconnect,
-                uint unknown0,
-                IntPtr user_info);
+			DirectoryInfo AppleMobileDeviceSupportX86 = new DirectoryInfo(Environment.GetFolderPath(Environment.SpecialFolder.CommonProgramFilesX86) + @"\Apple\Mobile Device Support");
+			if ((AppleMobileDeviceSupport != AppleMobileDeviceSupportX86) && (AppleMobileDeviceSupportX86.Exists))
+			{
+				PathBits.Add(AppleMobileDeviceSupportX86.FullName);
+			}
 
+			// Set the path from all the individual bits
+			Environment.SetEnvironmentVariable("Path", string.Join(";", PathBits));
+		}
+
+		public int NotificationSubscribe(DeviceNotificationCallback DeviceCallbackHandle)
+		{
+			return AMDeviceMethods.NotificationSubscribe(DeviceCallbackHandle);
+		}
+
+		public IntPtr CopyDeviceIdentifier(TypedPtr<AppleMobileDeviceConnection> device)
+		{
+			return AMDeviceMethods.CopyDeviceIdentifier(device);
+		}
+
+		public int Connect(TypedPtr<AppleMobileDeviceConnection> device)
+		{
+			return AMDeviceMethods.Connect(device);
+		}
+
+		public IntPtr CopyValue(TypedPtr<AppleMobileDeviceConnection> device, uint unknown, TypedPtr<CFString> cfstring)
+		{
+			return AMDeviceMethods.CopyValue(device, unknown, cfstring);
+		}
+
+		public int Disconnect(TypedPtr<AppleMobileDeviceConnection> device)
+		{
+			return AMDeviceMethods.Disconnect(device);
+		}
+
+		public int GetConnectionID(TypedPtr<AppleMobileDeviceConnection> device)
+		{
+			return AMDeviceMethods.GetConnectionID(device);
+		}
+
+		public int IsPaired(TypedPtr<AppleMobileDeviceConnection> device)
+		{
+			return AMDeviceMethods.IsPaired(device);
+		}
+
+		public int Pair(TypedPtr<AppleMobileDeviceConnection> device)
+		{
+			return AMDeviceMethods.Pair(device);
+		}
+
+		public int Unpair(TypedPtr<AppleMobileDeviceConnection> device)
+		{
+			return AMDeviceMethods.Unpair(device);
+		}
+
+		public int ValidatePairing(TypedPtr<AppleMobileDeviceConnection> device)
+		{
+			return AMDeviceMethods.ValidatePairing(device);
+		}
+
+		public int LookupApplications(TypedPtr<AppleMobileDeviceConnection> device, IntPtr options, out Dictionary<string, object> AppBundles)
+		{
+			int Result = AMDeviceMethods.LookupApplications(device, options, out AppBundles);
+			return Result;
+		}
+
+		public int StartSession(TypedPtr<AppleMobileDeviceConnection> device)
+		{
+			return AMDeviceMethods.StartSession((IntPtr)device);
+		}
+
+		public int StopSession(TypedPtr<AppleMobileDeviceConnection> device)
+		{
+			return AMDeviceMethods.StopSession((IntPtr)device);
+		}
+
+		public int StartHouseArrestService(TypedPtr<AppleMobileDeviceConnection> device, TypedPtr<CFString> bundleName, IntPtr unknown1, ref IntPtr handle, int unknown2)
+		{
+			return AMDeviceMethods.StartHouseArrestService((IntPtr)device, (IntPtr)bundleName, unknown1, ref handle, unknown2);
+		}
+
+		public int StartService(TypedPtr<AppleMobileDeviceConnection> device, TypedPtr<CFString> serviceName, ref IntPtr handle, IntPtr unknown)
+		{
+			return AMDeviceMethods.StartService((IntPtr)device, (IntPtr)serviceName, ref handle, unknown);
+		}
+
+		public int RestoreRegisterForDeviceNotifications(
+			DeviceRestoreNotificationCallback dfu_connect,
+			DeviceRestoreNotificationCallback recovery_connect,
+			DeviceRestoreNotificationCallback dfu_disconnect,
+			DeviceRestoreNotificationCallback recovery_disconnect,
+			uint unknown0,
+			IntPtr user_info)
+		{
+			return AMDeviceMethods.AMRestoreRegisterForDeviceNotifications(dfu_connect, recovery_connect, dfu_disconnect, recovery_disconnect, unknown0, user_info);
+		}
+
+		public int SecureUninstallApplication(
+			IntPtr serviceConnection,
+			TypedPtr<AppleMobileDeviceConnection> DeviceIfConnIsNull,
+			TypedPtr<CFString> ApplicationIdentifer,
+			TypedPtr<CFDictionary> ClientOptions,
+			DeviceInstallationCallback ProgressCallback,
+			IntPtr UserData)
+		{
+			return AMDeviceMethods.SecureUninstallApplication(serviceConnection, DeviceIfConnIsNull, ApplicationIdentifer, ClientOptions, ProgressCallback, UserData);
+		}
+
+		public int SecureInstallApplication(
+			IntPtr serviceConnection,
+			TypedPtr<AppleMobileDeviceConnection> DeviceIfConnIsNull,
+			TypedPtr<CFURL> UrlPath,
+			TypedPtr<CFDictionary> ClientOptions,
+			DeviceInstallationCallback ProgressCallback,
+			IntPtr UserData)
+		{
+			return AMDeviceMethods.SecureInstallApplication(serviceConnection, DeviceIfConnIsNull, UrlPath, ClientOptions, ProgressCallback, UserData);
+		}
+
+		public int SecureUpgradeApplication(
+			IntPtr ServiceConnection,
+			TypedPtr<AppleMobileDeviceConnection> DeviceIfConnIsNull,
+			TypedPtr<CFURL> UrlPath,
+			TypedPtr<CFDictionary> ClientOptions,
+			DeviceInstallationCallback ProgressCallback,
+			IntPtr UserData)
+		{
+			return AMDeviceMethods.SecureUpgradeApplication(ServiceConnection, DeviceIfConnIsNull, UrlPath, ClientOptions, ProgressCallback, UserData);
+		}
+
+		public int ConnectionClose(TypedPtr<AFCCommConnection> conn)
+		{
+			return AFC.ConnectionClose(conn);
+		}
+
+		public int ConnectionInvalidate(TypedPtr<AFCCommConnection> conn)
+		{
+			return AFC.ConnectionInvalidate(conn);
+		}
+
+		public int ConnectionIsValid(TypedPtr<AFCCommConnection> conn)
+		{
+			return AFC.ConnectionIsValid(conn);
+		}
+
+		public int ConnectionOpen(IntPtr handle, uint io_timeout, out TypedPtr<AFCCommConnection> OutConn)
+		{
+			return AFC.ConnectionOpen(handle, io_timeout, out OutConn);
+		}
+
+		public int DeviceInfoOpen(IntPtr handle, ref IntPtr dict)
+		{
+			return AFC.DeviceInfoOpen(handle, ref dict);
+		}
+
+		public int DirectoryClose(TypedPtr<AFCCommConnection> conn, IntPtr dir)
+		{
+			return AFC.DirectoryClose(conn, dir);
+		}
+
+		public int DirectoryCreate(TypedPtr<AFCCommConnection> conn, string path)
+		{
+			return AFC.DirectoryCreate(conn, path);
+		}
+
+		public int DirectoryOpen(TypedPtr<AFCCommConnection> conn, byte[] path, ref IntPtr dir)
+		{
+			return AFC.DirectoryOpen(conn, path, ref dir);
+		}
+
+		public int DirectoryRead(TypedPtr<AFCCommConnection> conn, IntPtr dir, ref IntPtr dirent)
+		{
+			return AFC.DirectoryRead(conn, dir, ref dirent);
+		}
+
+		public int FileInfoOpen(TypedPtr<AFCCommConnection> conn, string path, out TypedPtr<AFCDictionary> OutDict)
+		{
+			return AFC.FileInfoOpen(conn, path, out OutDict);
+		}
+
+		public int FileRefClose(TypedPtr<AFCCommConnection> conn, Int64 handle)
+		{
+			return AFC.FileRefClose(conn, handle);
+		}
+
+		public int FileRefOpen(TypedPtr<AFCCommConnection> conn, string path, int mode, int unknown, out Int64 handle)
+		{
+			return AFC.FileRefOpen(conn, path, mode, unknown, out handle);
+		}
+
+		public int FileRefRead(TypedPtr<AFCCommConnection> conn, Int64 handle, byte[] buffer, ref uint len)
+		{
+			return AFC.FileRefRead(conn, handle, buffer, ref len);
+		}
+
+		public int FileRefSeek(TypedPtr<AFCCommConnection> conn, Int64 handle, Int64 pos, Int64 origin)
+		{
+			return AFC.FileRefSeek(conn, handle, pos, origin);
+		}
+
+		public int FileRefSetFileSize(TypedPtr<AFCCommConnection> conn, Int64 handle, uint size)
+		{
+			return AFC.FileRefSetFileSize(conn, handle, size);
+		}
+
+		public int FileRefTell(TypedPtr<AFCCommConnection> conn, Int64 handle, ref uint position)
+		{
+			return AFC.FileRefTell(conn, handle, ref position);
+		}
+
+		public int FileRefWrite(TypedPtr<AFCCommConnection> conn, Int64 handle, byte[] buffer, uint len)
+		{
+			return AFC.FileRefWrite(conn, handle, buffer, len);
+		}
+
+		public int FlushData(TypedPtr<AFCCommConnection> conn, Int64 handle)
+		{
+			return AFC.FlushData(conn, handle);
+		}
+
+		public int KeyValueClose(TypedPtr<AFCDictionary> dict)
+		{
+			return AFC.KeyValueClose(dict);
+		}
+
+		public int KeyValueRead(TypedPtr<AFCDictionary> dict, out IntPtr key, out IntPtr val)
+		{
+			return AFC.KeyValueRead(dict, out key, out val);
+		}
+
+		public int RemovePath(TypedPtr<AFCCommConnection> conn, string path)
+		{
+			return AFC.RemovePath(conn, path);
+		}
+
+		public int RenamePath(TypedPtr<AFCCommConnection> conn, string OldPath, string NewPath)
+		{
+			return AFC.RenamePath(conn, OldPath, NewPath);
+		}
+
+		#region Application Methods
+		public class AMDeviceMethods
+		{
+			// Need to have the path manipulation code in static MobileDevice() to run
+			//			static int ForceMobileDeviceConstructorToRun;
+			static AMDeviceMethods()
+			{
+				//				ForceMobileDeviceConstructorToRun = ForceStaticInit;
+			}
 
-            [DllImport(DLLName, CallingConvention = CallingConvention.Cdecl)]
-            private extern static int AMDeviceUninstallApplication(
-                IntPtr/*AppleMobileDeviceConnection*/ device,
-                IntPtr/*CFString*/ ApplicationIdentifier,
-                IntPtr/*CFDictionary*/ ClientOptions,
-                DeviceInstallationCallback ProgressCallback,
-                IntPtr UserData);
+			public static IntPtr CopyDeviceIdentifier(TypedPtr<AppleMobileDeviceConnection> device)
+			{
+				return AMDeviceCopyDeviceIdentifier((IntPtr)device);
+			}
 
-            [DllImport(DLLName, CallingConvention = CallingConvention.Cdecl)]
-            private extern static int AMDeviceSecureUninstallApplication(
-                IntPtr serviceConnection,
-                IntPtr/*AppleMobileDeviceConnection*/ DeviceIfConnIsNull,
-                IntPtr/*CFString*/ ApplicationIdentifer,
-                IntPtr/*CFDictionary*/ ClientOptions,
-                DeviceInstallationCallback ProgressCallback,
-                IntPtr UserData);
+			public static int Connect(TypedPtr<AppleMobileDeviceConnection> device)
+			{
+				return AMDeviceConnect((IntPtr)device);
+			}
 
-            [DllImport(DLLName, CallingConvention = CallingConvention.Cdecl)]
-            private extern static int AMDeviceSecureInstallApplication(
-                IntPtr serviceConnection,
-                IntPtr/*AppleMobileDeviceConnection*/ DeviceIfConnIsNull,
-                IntPtr/*CFURL*/ UrlPath,
-                IntPtr/*CFDictionary*/ ClientOptions,
-                DeviceInstallationCallback ProgressCallback,
-                IntPtr UserData);
+			public static IntPtr CopyValue(TypedPtr<AppleMobileDeviceConnection> device, uint unknown, TypedPtr<CFString> cfstring)
+			{
+				return AMDeviceCopyValue((IntPtr)device, unknown, (IntPtr)cfstring);
+			}
 
-            [DllImport(DLLName, CallingConvention = CallingConvention.Cdecl)]
-            private extern static int AMDeviceTransferApplication(
-                IntPtr/*AppleMobileDeviceConnection*/ device,
-                IntPtr/*CFString*/ InPackagePath,
-                IntPtr UnknownButUnused,
-                DeviceInstallationCallback ProgressCallback,
-                IntPtr UserData);
+			public static int Disconnect(TypedPtr<AppleMobileDeviceConnection> device)
+			{
+				return AMDeviceDisconnect((IntPtr)device);
+			}
 
-            [DllImport(DLLName, CallingConvention = CallingConvention.Cdecl)]
-            private extern static int AMDeviceSecureUpgradeApplication(
-                IntPtr ServiceConnection,
-                IntPtr/*AppleMobileDeviceConnection*/ DeviceIfConnIsNull,
-                IntPtr/*CFURL*/ UrlPath,
-                IntPtr/*CFDictionary*/ ClientOptions,
-                DeviceInstallationCallback ProgressCallback,
-                IntPtr UserData);
+			public static int GetConnectionID(TypedPtr<AppleMobileDeviceConnection> device)
+			{
+				return AMDeviceGetConnectionID((IntPtr)device);
+			}
 
-            [DllImport(DLLName, CallingConvention = CallingConvention.Cdecl)]
-            private extern static int AMDeviceInstallApplication(
-                IntPtr/*AppleMobileDeviceConnection*/ device,
-                IntPtr/*CFString*/ FilePath,
-                IntPtr/*CFDictionary*/ ClientOptions,
-                DeviceInstallationCallback ProgressCallback,
-                IntPtr UserData);
+			public static int IsPaired(TypedPtr<AppleMobileDeviceConnection> device)
+			{
+				return AMDeviceIsPaired((IntPtr)device);
+			}
 
-            // the unknown goes into the 4th dword in the connection struct.  if non-zero, winsock send() won't get called in AMDServiceConnectionSend AFAIK
-            // the only option is "CloseOnInvalidate" (CFBoolean type, defaults to false)
-            [DllImport(DLLName, CallingConvention = CallingConvention.Cdecl)]
-            private extern static IntPtr AMDServiceConnectionCreate(IntPtr/*AppleMobileDeviceConnection*/ device, IntPtr UnknownTypically0, IntPtr OptionsDictionary);
+			public static int Pair(TypedPtr<AppleMobileDeviceConnection> device)
+			{
+				return AMDevicePair((IntPtr)device);
+			}
 
+			public static int Unpair(TypedPtr<AppleMobileDeviceConnection> device)
+			{
+				return AMDeviceUnpair((IntPtr)device);
+			}
 
-            [DllImport(DLLName, CallingConvention = CallingConvention.Cdecl)]
-            private extern static int AMSBeginSync(IntPtr/*AppleMobileDeviceConnection*/ device, IntPtr OuterIn8, IntPtr OuterIn12, IntPtr OuterIn16);
-        }
-#endregion
+			public static int ValidatePairing(TypedPtr<AppleMobileDeviceConnection> device)
+			{
+				return AMDeviceValidatePairing((IntPtr)device);
+			}
 
+			public static int LookupApplications(TypedPtr<AppleMobileDeviceConnection> device, IntPtr options, out Dictionary<string, object> AppBundles)
+			{
+				IntPtr UntypedDict = IntPtr.Zero;
+				int Result = AMDeviceLookupApplications((IntPtr)device, options, ref UntypedDict);
 
-#region AFC Operations
+				AppBundles = MobileDevice.ConvertCFDictionaryToDictionaryString(new TypedPtr<CFDictionary>(UntypedDict));
 
+				return Result;
+			}
 
+			public static int StartSession(TypedPtr<AppleMobileDeviceConnection> device)
+			{
+				return AMDeviceStartSession((IntPtr)device);
+			}
 
+			public static int StopSession(TypedPtr<AppleMobileDeviceConnection> device)
+			{
+				return AMDeviceStopSession((IntPtr)device);
+			}
 
+			public static int StartHouseArrestService(TypedPtr<AppleMobileDeviceConnection> device, TypedPtr<CFString> bundleName, IntPtr unknown1, ref IntPtr handle, int unknown2)
+			{
+				return AMDeviceStartHouseArrestService((IntPtr)device, (IntPtr)bundleName, unknown1, ref handle, unknown2);
+			}
 
+			public static int StartService(TypedPtr<AppleMobileDeviceConnection> device, TypedPtr<CFString> serviceName, ref IntPtr handle, IntPtr unknown)
+			{
+				return AMDeviceStartService((IntPtr)device, (IntPtr)serviceName, ref handle, unknown);
+			}
 
-        public static int AFCDirectoryOpen(TypedPtr<AFCCommConnection> conn, string path, ref IntPtr dir)
-        {
-            return AFC.DirectoryOpen((IntPtr)conn, Encoding.UTF8.GetBytes(path), ref dir);
-        }
+			public static int UninstallApplication(
+				TypedPtr<AppleMobileDeviceConnection> device,
+				TypedPtr<CFString> ApplicationIdentifier,
+				TypedPtr<CFDictionary> ClientOptions,
+				DeviceInstallationCallback ProgressCallback,
+				IntPtr UserData)
+			{
+				return AMDeviceUninstallApplication((IntPtr)device, (IntPtr)ApplicationIdentifier, (IntPtr)ClientOptions, ProgressCallback, UserData);
+			}
 
-        public static int AFCDirectoryRead(TypedPtr<AFCCommConnection> conn, IntPtr dir, ref string buffer)
-        {
-            int ret;
+			public static int SecureUninstallApplication(
+				IntPtr serviceConnection,
+				TypedPtr<AppleMobileDeviceConnection> DeviceIfConnIsNull,
+				TypedPtr<CFString> ApplicationIdentifer,
+				TypedPtr<CFDictionary> ClientOptions,
+				DeviceInstallationCallback ProgressCallback,
+				IntPtr UserData)
+			{
+				return AMDeviceSecureUninstallApplication(serviceConnection, (IntPtr)DeviceIfConnIsNull, (IntPtr)ApplicationIdentifer, (IntPtr)ClientOptions, ProgressCallback, UserData);
+			}
 
-            IntPtr ptr = IntPtr.Zero;
-            ret = AFC.DirectoryRead((IntPtr)conn, dir, ref ptr);
-            if ((ret == 0) && (ptr != null))
-            {
-                buffer = Marshal.PtrToStringAnsi(ptr);
-            }
-            else
-            {
-                buffer = null;
-            }
-            return ret;
-        }
+			public static int SecureInstallApplication(
+				IntPtr serviceConnection,
+				TypedPtr<AppleMobileDeviceConnection> DeviceIfConnIsNull,
+				TypedPtr<CFURL> UrlPath,
+				TypedPtr<CFDictionary> ClientOptions,
+				DeviceInstallationCallback ProgressCallback,
+				IntPtr UserData)
+			{
+				return AMDeviceSecureInstallApplication(serviceConnection, (IntPtr)DeviceIfConnIsNull, (IntPtr)UrlPath, (IntPtr)ClientOptions, ProgressCallback, UserData);
+			}
 
-        public class AFC
-        {
-            public static int ConnectionClose(TypedPtr<AFCCommConnection> conn)
-            {
-                return AFCConnectionClose((IntPtr)conn);
-            }
+			public static int TransferApplication(
+				TypedPtr<AppleMobileDeviceConnection> device,
+				TypedPtr<CFString> InPackagePath,
+				IntPtr UnknownButUnused,
+				DeviceInstallationCallback ProgressCallback,
+				IntPtr UserData)
+			{
+				return AMDeviceTransferApplication((IntPtr)device, (IntPtr)InPackagePath, UnknownButUnused, ProgressCallback, UserData);
+			}
 
-            public static int ConnectionInvalidate(TypedPtr<AFCCommConnection> conn)
-            {
-                return AFCConnectionInvalidate((IntPtr)conn);
-            }
+			public static int SecureUpgradeApplication(
+				IntPtr ServiceConnection,
+				TypedPtr<AppleMobileDeviceConnection> DeviceIfConnIsNull,
+				TypedPtr<CFURL> UrlPath,
+				TypedPtr<CFDictionary> ClientOptions,
+				DeviceInstallationCallback ProgressCallback,
+				IntPtr UserData)
+			{
+				return AMDeviceSecureUpgradeApplication(ServiceConnection, (IntPtr)DeviceIfConnIsNull, (IntPtr)UrlPath, (IntPtr)ClientOptions, ProgressCallback, UserData);
+			}
 
-            public static int ConnectionIsValid(TypedPtr<AFCCommConnection> conn)
-            {
-                return AFCConnectionIsValid((IntPtr)conn);
-            }
+			public static int InstallApplication(
+				TypedPtr<AppleMobileDeviceConnection> device,
+				TypedPtr<CFString> FilePath,
+				TypedPtr<CFDictionary> ClientOptions,
+				DeviceInstallationCallback ProgressCallback,
+				IntPtr UserData)
+			{
+				return AMDeviceInstallApplication((IntPtr)device, (IntPtr)FilePath, (IntPtr)ClientOptions, ProgressCallback, UserData);
+			}
 
-            public static int ConnectionOpen(IntPtr handle, uint io_timeout, out TypedPtr<AFCCommConnection> OutConn)
-            {
-                IntPtr Conn;
-                int Result = AFCConnectionOpen(handle, io_timeout, out Conn);
-                OutConn = Conn;
-                return Result;
-            }
+			public static int NotificationSubscribe(DeviceNotificationCallback DeviceCallbackHandle)
+			{
+				IntPtr notification;
+				return AMDeviceNotificationSubscribe(DeviceCallbackHandle, 0, 0, 0, out notification);
+			}
 
-            public static int DeviceInfoOpen(IntPtr handle, ref IntPtr dict)
-            {
-                return AFCDeviceInfoOpen(handle, ref dict);
-            }
+			[DllImport(DLLName, CallingConvention = CallingConvention.Cdecl)]
+			private extern static IntPtr AMDeviceCopyDeviceIdentifier(IntPtr/*AppleMobileDeviceConnection*/ device);
+
+			[DllImport(DLLName, CallingConvention = CallingConvention.Cdecl)]
+			private extern static int AMDeviceConnect(IntPtr/*AppleMobileDeviceConnection*/ device);
+
+			[DllImport(DLLName, CallingConvention = CallingConvention.Cdecl)]
+			private extern static IntPtr AMDeviceCopyValue(IntPtr/*AppleMobileDeviceConnection*/ device, uint unknown, IntPtr/*CFString*/ cfstring);
+
+			[DllImport(DLLName, CallingConvention = CallingConvention.Cdecl)]
+			private extern static int AMDeviceDisconnect(IntPtr/*AppleMobileDeviceConnection*/ device);
+
+			[DllImport(DLLName, CallingConvention = CallingConvention.Cdecl)]
+			private extern static int AMDeviceGetConnectionID(IntPtr/*AppleMobileDeviceConnection*/ device);
+
+			[DllImport(DLLName, CallingConvention = CallingConvention.Cdecl)]
+			private extern static int AMDeviceIsPaired(IntPtr/*AppleMobileDeviceConnection*/ device);
+
+			[DllImport(DLLName, CallingConvention = CallingConvention.Cdecl)]
+			private extern static int AMDevicePair(IntPtr/*AppleMobileDeviceConnection*/ device);
+
+			[DllImport(DLLName, CallingConvention = CallingConvention.Cdecl)]
+			private extern static int AMDeviceUnpair(IntPtr/*AppleMobileDeviceConnection*/ device);
+
+			[DllImport(DLLName, CallingConvention = CallingConvention.Cdecl)]
+			private extern static int AMDeviceValidatePairing(IntPtr/*AppleMobileDeviceConnection*/ device);
+
+			[DllImport(DLLName, CallingConvention = CallingConvention.Cdecl)]
+			private extern static int AMDeviceLookupApplications(IntPtr/*AppleMobileDeviceConnection*/ device, IntPtr options, ref IntPtr appBundles);
+
+			[DllImport(DLLName, CallingConvention = CallingConvention.Cdecl)]
+			private extern static int AMDeviceNotificationSubscribe(DeviceNotificationCallback callback, uint unused1, uint unused2, uint unused3, out IntPtr am_device_notification_ptr);
+
+			[DllImport(DLLName, CallingConvention = CallingConvention.Cdecl)]
+			private extern static int AMDeviceStartHouseArrestService(IntPtr/*AppleMobileDeviceConnection*/ device, IntPtr/*CFString*/ bundleName, IntPtr unknown1, ref IntPtr handle, int unknown2);
+
+			[DllImport(DLLName, CallingConvention = CallingConvention.Cdecl)]
+			private extern static int AMDeviceStartService(IntPtr/*AppleMobileDeviceConnection*/ device, IntPtr/*CFString*/ serviceName, ref IntPtr handle, IntPtr unknown);
+
+			[DllImport(DLLName, CallingConvention = CallingConvention.Cdecl)]
+			private extern static int AMDeviceStartSession(IntPtr/*AppleMobileDeviceConnection*/ device);
+
+			[DllImport(DLLName, CallingConvention = CallingConvention.Cdecl)]
+			private extern static int AMDeviceStopSession(IntPtr/*AppleMobileDeviceConnection*/ device);
+
+			[DllImport(DLLName, CallingConvention = CallingConvention.Cdecl)]
+			private extern static int AMRestoreModeDeviceCreate(uint unknown0, int connection_id, uint unknown1);
+
+			[DllImport(DLLName, CallingConvention = CallingConvention.Cdecl)]
+			public extern static int AMRestoreRegisterForDeviceNotifications(
+				DeviceRestoreNotificationCallback dfu_connect,
+				DeviceRestoreNotificationCallback recovery_connect,
+				DeviceRestoreNotificationCallback dfu_disconnect,
+				DeviceRestoreNotificationCallback recovery_disconnect,
+				uint unknown0,
+				IntPtr user_info);
+
+			[DllImport(DLLName, CallingConvention = CallingConvention.Cdecl)]
+			private extern static int AMDeviceUninstallApplication(
+				IntPtr/*AppleMobileDeviceConnection*/ device,
+				IntPtr/*CFString*/ ApplicationIdentifier,
+				IntPtr/*CFDictionary*/ ClientOptions,
+				DeviceInstallationCallback ProgressCallback,
+				IntPtr UserData);
+
+			[DllImport(DLLName, CallingConvention = CallingConvention.Cdecl)]
+			private extern static int AMDeviceSecureUninstallApplication(
+				IntPtr serviceConnection,
+				IntPtr/*AppleMobileDeviceConnection*/ DeviceIfConnIsNull,
+				IntPtr/*CFString*/ ApplicationIdentifer,
+				IntPtr/*CFDictionary*/ ClientOptions,
+				DeviceInstallationCallback ProgressCallback,
+				IntPtr UserData);
+
+			[DllImport(DLLName, CallingConvention = CallingConvention.Cdecl)]
+			private extern static int AMDeviceSecureInstallApplication(
+				IntPtr serviceConnection,
+				IntPtr/*AppleMobileDeviceConnection*/ DeviceIfConnIsNull,
+				IntPtr/*CFURL*/ UrlPath,
+				IntPtr/*CFDictionary*/ ClientOptions,
+				DeviceInstallationCallback ProgressCallback,
+				IntPtr UserData);
+
+			[DllImport(DLLName, CallingConvention = CallingConvention.Cdecl)]
+			private extern static int AMDeviceTransferApplication(
+				IntPtr/*AppleMobileDeviceConnection*/ device,
+				IntPtr/*CFString*/ InPackagePath,
+				IntPtr UnknownButUnused,
+				DeviceInstallationCallback ProgressCallback,
+				IntPtr UserData);
+
+			[DllImport(DLLName, CallingConvention = CallingConvention.Cdecl)]
+			private extern static int AMDeviceSecureUpgradeApplication(
+				IntPtr ServiceConnection,
+				IntPtr/*AppleMobileDeviceConnection*/ DeviceIfConnIsNull,
+				IntPtr/*CFURL*/ UrlPath,
+				IntPtr/*CFDictionary*/ ClientOptions,
+				DeviceInstallationCallback ProgressCallback,
+				IntPtr UserData);
+
+			[DllImport(DLLName, CallingConvention = CallingConvention.Cdecl)]
+			private extern static int AMDeviceInstallApplication(
+				IntPtr/*AppleMobileDeviceConnection*/ device,
+				IntPtr/*CFString*/ FilePath,
+				IntPtr/*CFDictionary*/ ClientOptions,
+				DeviceInstallationCallback ProgressCallback,
+				IntPtr UserData);
+
+			// the unknown goes into the 4th dword in the connection struct.  if non-zero, winsock send() won't get called in AMDServiceConnectionSend AFAIK
+			// the only option is "CloseOnInvalidate" (CFBoolean type, defaults to false)
+			[DllImport(DLLName, CallingConvention = CallingConvention.Cdecl)]
+			private extern static IntPtr AMDServiceConnectionCreate(IntPtr/*AppleMobileDeviceConnection*/ device, IntPtr UnknownTypically0, IntPtr OptionsDictionary);
+
+
+			[DllImport(DLLName, CallingConvention = CallingConvention.Cdecl)]
+			private extern static int AMSBeginSync(IntPtr/*AppleMobileDeviceConnection*/ device, IntPtr OuterIn8, IntPtr OuterIn12, IntPtr OuterIn16);
+		}
+		#endregion
+
+		#region AFC Operations
+
+		public int DirectoryOpen(TypedPtr<AFCCommConnection> conn, string path, ref IntPtr dir)
+		{
+			return AFC.DirectoryOpen((IntPtr)conn, Encoding.UTF8.GetBytes(path), ref dir);
+		}
+
+		public int DirectoryRead(TypedPtr<AFCCommConnection> conn, IntPtr dir, ref string buffer)
+		{
+			int ret;
+
+			IntPtr ptr = IntPtr.Zero;
+			ret = AFC.DirectoryRead((IntPtr)conn, dir, ref ptr);
+			if ((ret == 0) && (ptr != IntPtr.Zero))
+			{
+				buffer = Marshal.PtrToStringAnsi(ptr);
+			}
+			else
+			{
+				buffer = null;
+			}
+			return ret;
+		}
 
-            public static int DirectoryClose(TypedPtr<AFCCommConnection> conn, IntPtr dir)
-            {
-                return AFCDirectoryClose((IntPtr)conn, dir);
-            }
+		public class AFC
+		{
+			public static int ConnectionClose(TypedPtr<AFCCommConnection> conn)
+			{
+				return AFCConnectionClose((IntPtr)conn);
+			}
 
-            public static int DirectoryCreate(TypedPtr<AFCCommConnection> conn, string path)
-            {
-                return AFCDirectoryCreate((IntPtr)conn, path);
-            }
+			public static int ConnectionInvalidate(TypedPtr<AFCCommConnection> conn)
+			{
+				return AFCConnectionInvalidate((IntPtr)conn);
+			}
 
-            public static int DirectoryOpen(TypedPtr<AFCCommConnection> conn, byte[] path, ref IntPtr dir)
-            {
-                return AFCDirectoryOpen((IntPtr)conn, path, ref dir);
-            }
+			public static int ConnectionIsValid(TypedPtr<AFCCommConnection> conn)
+			{
+				return AFCConnectionIsValid((IntPtr)conn);
+			}
 
-            public static int DirectoryRead(TypedPtr<AFCCommConnection> conn, IntPtr dir, ref IntPtr dirent)
-            {
-                return AFCDirectoryRead((IntPtr)conn, dir, ref dirent);
-            }
+			public static int ConnectionOpen(IntPtr handle, uint io_timeout, out TypedPtr<AFCCommConnection> OutConn)
+			{
+				IntPtr Conn;
+				int Result = AFCConnectionOpen(handle, io_timeout, out Conn);
+				OutConn = Conn;
+				return Result;
+			}
 
-            public static int FileInfoOpen(TypedPtr<AFCCommConnection> conn, string path, out TypedPtr<AFCDictionary> OutDict)
-            {
-                IntPtr UntypedDict;
-                int Result = AFCFileInfoOpen((IntPtr)conn, path, out UntypedDict);
-                OutDict = UntypedDict;
+			public static int DeviceInfoOpen(IntPtr handle, ref IntPtr dict)
+			{
+				return AFCDeviceInfoOpen(handle, ref dict);
+			}
 
-                return Result;
-            }
+			public static int DirectoryClose(TypedPtr<AFCCommConnection> conn, IntPtr dir)
+			{
+				return AFCDirectoryClose((IntPtr)conn, dir);
+			}
 
-            public static int FileRefClose(TypedPtr<AFCCommConnection> conn, Int64 handle)
-            {
-                return AFCFileRefClose((IntPtr)conn, handle);
-            }
+			public static int DirectoryCreate(TypedPtr<AFCCommConnection> conn, string path)
+			{
+				return AFCDirectoryCreate((IntPtr)conn, path);
+			}
 
-            public static int FileRefOpen(TypedPtr<AFCCommConnection> conn, string path, int mode, int unknown, out Int64 handle)
-            {
-                return AFCFileRefOpen((IntPtr)conn, path, mode, unknown, out handle);
-            }
+			public static int DirectoryOpen(TypedPtr<AFCCommConnection> conn, byte[] path, ref IntPtr dir)
+			{
+				return AFCDirectoryOpen((IntPtr)conn, path, ref dir);
+			}
 
-            public static int FileRefRead(TypedPtr<AFCCommConnection> conn, Int64 handle, byte[] buffer, ref uint len)
-            {
-                return AFCFileRefRead((IntPtr)conn, handle, buffer, ref len);
-            }
+			public static int DirectoryRead(TypedPtr<AFCCommConnection> conn, IntPtr dir, ref IntPtr dirent)
+			{
+				return AFCDirectoryRead((IntPtr)conn, dir, ref dirent);
+			}
 
-            public static int FileRefSeek(TypedPtr<AFCCommConnection> conn, Int64 handle, Int64 pos, Int64 origin)
-            {
-                return AFCFileRefSeek((IntPtr)conn, handle, pos, origin);
-            }
+			public static int FileInfoOpen(TypedPtr<AFCCommConnection> conn, string path, out TypedPtr<AFCDictionary> OutDict)
+			{
+				IntPtr UntypedDict;
+				int Result = AFCFileInfoOpen((IntPtr)conn, path, out UntypedDict);
+				OutDict = UntypedDict;
 
-            public static int FileRefTell(TypedPtr<AFCCommConnection> conn, Int64 handle, ref uint position)
-            {
-                return AFCFileRefTell((IntPtr)conn, handle, ref position);
-            }
+				return Result;
+			}
 
-            public static int FileRefWrite(TypedPtr<AFCCommConnection> conn, Int64 handle, byte[] buffer, uint len)
-            {
-                return AFCFileRefWrite((IntPtr)conn, handle, buffer, len);
-            }
+			public static int FileRefClose(TypedPtr<AFCCommConnection> conn, Int64 handle)
+			{
+				return AFCFileRefClose((IntPtr)conn, handle);
+			}
 
-            public static int FlushData(TypedPtr<AFCCommConnection> conn, Int64 handle)
-            {
-                return AFCFlushData((IntPtr)conn, handle);
-            }
+			public static int FileRefOpen(TypedPtr<AFCCommConnection> conn, string path, int mode, int unknown, out Int64 handle)
+			{
+				return AFCFileRefOpen((IntPtr)conn, path, mode, unknown, out handle);
+			}
 
-            public static int KeyValueClose(TypedPtr<AFCDictionary> dict)
-            {
-                return AFCKeyValueClose((IntPtr)dict);
-            }
+			public static int FileRefRead(TypedPtr<AFCCommConnection> conn, Int64 handle, byte[] buffer, ref uint len)
+			{
+				return AFCFileRefRead((IntPtr)conn, handle, buffer, ref len);
+			}
 
-            public static int KeyValueRead(TypedPtr<AFCDictionary> dict, out IntPtr key, out IntPtr val)
-            {
-                return AFCKeyValueRead((IntPtr)dict, out key, out val);
-            }
+			public static int FileRefSetFileSize(TypedPtr<AFCCommConnection> conn, Int64 handle, uint size)
+			{
+				return AFCFileRefSetFileSize((IntPtr)conn, handle, size);
+			}
 
-            public static int RemovePath(TypedPtr<AFCCommConnection> conn, string path)
-            {
-                return AFCRemovePath((IntPtr)conn, path);
-            }
+			public static int FileRefSeek(TypedPtr<AFCCommConnection> conn, Int64 handle, Int64 pos, Int64 origin)
+			{
+				return AFCFileRefSeek((IntPtr)conn, handle, pos, origin);
+			}
 
-            public static int RenamePath(TypedPtr<AFCCommConnection> conn, string OldPath, string NewPath)
-            {
-                return AFCRenamePath((IntPtr)conn, OldPath, NewPath);
-            }
+			public static int FileRefTell(TypedPtr<AFCCommConnection> conn, Int64 handle, ref uint position)
+			{
+				return AFCFileRefTell((IntPtr)conn, handle, ref position);
+			}
 
+			public static int FileRefWrite(TypedPtr<AFCCommConnection> conn, Int64 handle, byte[] buffer, uint len)
+			{
+				return AFCFileRefWrite((IntPtr)conn, handle, buffer, len);
+			}
 
+			public static int FlushData(TypedPtr<AFCCommConnection> conn, Int64 handle)
+			{
+				return AFCFlushData((IntPtr)conn, handle);
+			}
 
-            [DllImport(DLLName, CallingConvention = CallingConvention.Cdecl)]
-            private extern static int AFCConnectionClose(IntPtr/*AFCCommConnection*/ conn);
+			public static int KeyValueClose(TypedPtr<AFCDictionary> dict)
+			{
+				return AFCKeyValueClose((IntPtr)dict);
+			}
 
-            [DllImport(DLLName, CallingConvention = CallingConvention.Cdecl)]
-            private extern static int AFCConnectionInvalidate(IntPtr/*AFCCommConnection*/ conn);
+			public static int KeyValueRead(TypedPtr<AFCDictionary> dict, out IntPtr key, out IntPtr val)
+			{
+				return AFCKeyValueRead((IntPtr)dict, out key, out val);
+			}
 
-            [DllImport(DLLName, CallingConvention = CallingConvention.Cdecl)]
-            private extern static int AFCConnectionIsValid(IntPtr/*AFCCommConnection*/ conn);
+			public static int RemovePath(TypedPtr<AFCCommConnection> conn, string path)
+			{
+				return AFCRemovePath((IntPtr)conn, path);
+			}
 
-            [DllImport(DLLName, CallingConvention = CallingConvention.Cdecl)]
-            private extern static int AFCConnectionOpen(IntPtr handle, uint io_timeout, out IntPtr/*AFCCommConnection*/ conn);
+			public static int RenamePath(TypedPtr<AFCCommConnection> conn, string OldPath, string NewPath)
+			{
+				return AFCRenamePath((IntPtr)conn, OldPath, NewPath);
+			}
 
+			[DllImport(DLLName, CallingConvention = CallingConvention.Cdecl)]
+			private extern static int AFCConnectionClose(IntPtr/*AFCCommConnection*/ conn);
 
-            [DllImport(DLLName, CallingConvention = CallingConvention.Cdecl)]
-            private extern static int AFCDeviceInfoOpen(IntPtr handle, ref IntPtr dict);
+			[DllImport(DLLName, CallingConvention = CallingConvention.Cdecl)]
+			private extern static int AFCConnectionInvalidate(IntPtr/*AFCCommConnection*/ conn);
 
+			[DllImport(DLLName, CallingConvention = CallingConvention.Cdecl)]
+			private extern static int AFCConnectionIsValid(IntPtr/*AFCCommConnection*/ conn);
 
-            [DllImport(DLLName, CallingConvention = CallingConvention.Cdecl)]
-            private extern static int AFCDirectoryClose(IntPtr/*AFCCommConnection*/ conn, IntPtr dir);
+			[DllImport(DLLName, CallingConvention = CallingConvention.Cdecl)]
+			private extern static int AFCConnectionOpen(IntPtr handle, uint io_timeout, out IntPtr/*AFCCommConnection*/ conn);
 
-            [DllImport(DLLName, CallingConvention = CallingConvention.Cdecl)]
-            private extern static int AFCDirectoryCreate(IntPtr/*AFCCommConnection*/ conn, string path);
+			[DllImport(DLLName, CallingConvention = CallingConvention.Cdecl)]
+			private extern static int AFCDeviceInfoOpen(IntPtr handle, ref IntPtr dict);
 
-            [DllImport(DLLName, CallingConvention = CallingConvention.Cdecl)]
-            private extern static int AFCDirectoryOpen(IntPtr/*AFCCommConnection*/ conn, byte[] path, ref IntPtr dir);
+			[DllImport(DLLName, CallingConvention = CallingConvention.Cdecl)]
+			private extern static int AFCDirectoryClose(IntPtr/*AFCCommConnection*/ conn, IntPtr dir);
 
-            [DllImport(DLLName, CallingConvention = CallingConvention.Cdecl)]
-            private extern static int AFCDirectoryRead(IntPtr/*AFCCommConnection*/ conn, IntPtr dir, ref IntPtr dirent);
+			[DllImport(DLLName, CallingConvention = CallingConvention.Cdecl)]
+			private extern static int AFCDirectoryCreate(IntPtr/*AFCCommConnection*/ conn, string path);
 
+			[DllImport(DLLName, CallingConvention = CallingConvention.Cdecl)]
+			private extern static int AFCDirectoryOpen(IntPtr/*AFCCommConnection*/ conn, byte[] path, ref IntPtr dir);
 
-            [DllImport(DLLName, CallingConvention = CallingConvention.Cdecl)]
-            private extern static int AFCFileInfoOpen(IntPtr/*AFCCommConnection*/ conn, string path, out IntPtr OutDict);
+			[DllImport(DLLName, CallingConvention = CallingConvention.Cdecl)]
+			private extern static int AFCDirectoryRead(IntPtr/*AFCCommConnection*/ conn, IntPtr dir, ref IntPtr dirent);
 
+			[DllImport(DLLName, CallingConvention = CallingConvention.Cdecl)]
+			private extern static int AFCFileInfoOpen(IntPtr/*AFCCommConnection*/ conn, string path, out IntPtr OutDict);
 
-            [DllImport(DLLName, CallingConvention = CallingConvention.Cdecl)]
-            private extern static int AFCFileRefClose(IntPtr/*AFCCommConnection*/ conn, Int64 handle);
+			[DllImport(DLLName, CallingConvention = CallingConvention.Cdecl)]
+			private extern static int AFCFileRefClose(IntPtr/*AFCCommConnection*/ conn, Int64 handle);
 
-            [DllImport(DLLName, CallingConvention = CallingConvention.Cdecl)]
-            private extern static int AFCFileRefOpen(IntPtr/*AFCCommConnection*/ conn, string path, int mode, int unknown, out Int64 handle);
+			[DllImport(DLLName, CallingConvention = CallingConvention.Cdecl)]
+			private extern static int AFCFileRefOpen(IntPtr/*AFCCommConnection*/ conn, string path, int mode, int unknown, out Int64 handle);
 
-            [DllImport(DLLName, CallingConvention = CallingConvention.Cdecl)]
-            private extern static int AFCFileRefRead(IntPtr/*AFCCommConnection*/ conn, Int64 handle, byte[] buffer, ref uint len);
+			[DllImport(DLLName, CallingConvention = CallingConvention.Cdecl)]
+			private extern static int AFCFileRefRead(IntPtr/*AFCCommConnection*/ conn, Int64 handle, byte[] buffer, ref uint len);
 
-            [DllImport(DLLName, CallingConvention = CallingConvention.Cdecl)]
-            private extern static int AFCFileRefSeek(IntPtr/*AFCCommConnection*/ conn, Int64 handle, Int64 pos, Int64 origin);
+			[DllImport(DLLName, CallingConvention = CallingConvention.Cdecl)]
+			private extern static int AFCFileRefSeek(IntPtr/*AFCCommConnection*/ conn, Int64 handle, Int64 pos, Int64 origin);
 
-            // FIXME - not working, arguments?
-            [DllImport(DLLName, CallingConvention = CallingConvention.Cdecl)]
-            public extern static int AFCFileRefSetFileSize(IntPtr/*AFCCommConnection*/ conn, Int64 handle, uint size);
+			// FIXME - not working, arguments?
+			[DllImport(DLLName, CallingConvention = CallingConvention.Cdecl)]
+			public extern static int AFCFileRefSetFileSize(IntPtr/*AFCCommConnection*/ conn, Int64 handle, uint size);
 
-            [DllImport(DLLName, CallingConvention = CallingConvention.Cdecl)]
-            private extern static int AFCFileRefTell(IntPtr/*AFCCommConnection*/ conn, Int64 handle, ref uint position);
+			[DllImport(DLLName, CallingConvention = CallingConvention.Cdecl)]
+			private extern static int AFCFileRefTell(IntPtr/*AFCCommConnection*/ conn, Int64 handle, ref uint position);
 
-            [DllImport(DLLName, CallingConvention = CallingConvention.Cdecl)]
-            private extern static int AFCFileRefWrite(IntPtr/*AFCCommConnection*/ conn, Int64 handle, byte[] buffer, uint len);
+			[DllImport(DLLName, CallingConvention = CallingConvention.Cdecl)]
+			private extern static int AFCFileRefWrite(IntPtr/*AFCCommConnection*/ conn, Int64 handle, byte[] buffer, uint len);
 
+			[DllImport(DLLName, CallingConvention = CallingConvention.Cdecl)]
+			private extern static int AFCFlushData(IntPtr/*AFCCommConnection*/ conn, Int64 handle);
 
-            [DllImport(DLLName, CallingConvention = CallingConvention.Cdecl)]
-            private extern static int AFCFlushData(IntPtr/*AFCCommConnection*/ conn, Int64 handle);
+			[DllImport(DLLName, CallingConvention = CallingConvention.Cdecl)]
+			private extern static int AFCKeyValueClose(IntPtr dict);
 
+			[DllImport(DLLName, CallingConvention = CallingConvention.Cdecl)]
+			private extern static int AFCKeyValueRead(IntPtr dict, out IntPtr key, out IntPtr val);
 
-            [DllImport(DLLName, CallingConvention = CallingConvention.Cdecl)]
-            private extern static int AFCKeyValueClose(IntPtr dict);
 
-            [DllImport(DLLName, CallingConvention = CallingConvention.Cdecl)]
-            private extern static int AFCKeyValueRead(IntPtr dict, out IntPtr key, out IntPtr val);
+			[DllImport(DLLName, CallingConvention = CallingConvention.Cdecl)]
+			private extern static int AFCRemovePath(IntPtr/*AFCCommConnection*/ conn, string path);
 
+			[DllImport(DLLName, CallingConvention = CallingConvention.Cdecl)]
+			private extern static int AFCRenamePath(IntPtr/*AFCCommConnection*/ conn, string old_path, string new_path);
+		}
 
-            [DllImport(DLLName, CallingConvention = CallingConvention.Cdecl)]
-            private extern static int AFCRemovePath(IntPtr/*AFCCommConnection*/ conn, string path);
+		#endregion
+	}
 
-            [DllImport(DLLName, CallingConvention = CallingConvention.Cdecl)]
-            private extern static int AFCRenamePath(IntPtr/*AFCCommConnection*/ conn, string old_path, string new_path);
-        }
+    internal partial class MobileDevice
+    {
+		static public MobileDeviceImpl DeviceImpl = null;
 
-#endregion
+        static MobileDevice()
+        {
+			if (Environment.OSVersion.Platform == PlatformID.MacOSX || Environment.OSVersion.Platform == PlatformID.Unix)
+			{
+				DeviceImpl = new MobileDeviceOSX();
+			}
+			else
+			{
+				DeviceImpl = new MobileDeviceWin();
+			}
 
+            // Initialize the CoreFoundation bindings
+            InitializeCoreFoundation();
+        }
 
         /*
     Valid Value Names:
@@ -724,7 +1733,7 @@ namespace Manzana
 
         public static string AMDeviceCopyValue(TypedPtr<AppleMobileDeviceConnection> device, string name)
         {
-            IntPtr result = AMDeviceMethods.CopyValue(device, 0, CFStringMakeConstantString(name));
+            IntPtr result = DeviceImpl.CopyValue(device, 0, CFStringMakeConstantString(name));
             if (result != IntPtr.Zero)
             {
                 return MobileDevice.CFStringGetCString(result);
diff --git a/Engine/Source/Programs/IOS/MobileDeviceInterface/MobileDeviceInstance.cs b/Engine/Source/Programs/IOS/MobileDeviceInterface/MobileDeviceInstance.cs
--- a/Engine/Source/Programs/IOS/MobileDeviceInterface/MobileDeviceInstance.cs
+++ b/Engine/Source/Programs/IOS/MobileDeviceInterface/MobileDeviceInstance.cs
@@ -96,7 +96,7 @@ namespace Manzana
         
             DeviceCallbackHandle = new DeviceNotificationCallback(NotifyCallback);
             
-            int ret = MobileDevice.AMDeviceMethods.NotificationSubscribe(DeviceCallbackHandle);
+            int ret = MobileDevice.DeviceImpl.NotificationSubscribe(DeviceCallbackHandle);
             if (ret != 0)
             {
                 throw new Exception("AMDeviceNotificationSubscribe failed with error " + ret);
@@ -195,7 +195,7 @@ namespace Manzana
             drn3 = new DeviceRestoreNotificationCallback(DfuDisconnectCallback);
             drn4 = new DeviceRestoreNotificationCallback(RecoveryDisconnectCallback);
 
-            int ret = MobileDevice.AMDeviceMethods.AMRestoreRegisterForDeviceNotifications(drn1, drn2, drn3, drn4, 0, IntPtr.Zero);
+			int ret = MobileDevice.DeviceImpl.RestoreRegisterForDeviceNotifications(drn1, drn2, drn3, drn4, 0, IntPtr.Zero);
             if (ret != 0)
             {
                 throw new Exception("AMRestoreRegisterForDeviceNotifications failed with error " + ret);
@@ -493,10 +493,32 @@ namespace Manzana
             }
         }
 
+		void RecursiveCopy(string SourceFolderOnPC, string TargetFolderOnDevice)
+		{
+			string[] Directories = System.IO.Directory.GetDirectories(SourceFolderOnPC);
+			foreach (string Directory in Directories)
+			{
+				string NewSourceFolder = Directory;
+				string NewTargetFolder = TargetFolderOnDevice + Directory.Substring(Directory.LastIndexOf(Path.DirectorySeparatorChar)+1) + "/";
+
+				WriteProgressLine("Copying folder {0} to {1}", 0, NewSourceFolder, NewTargetFolder);
+				RecursiveCopy(NewSourceFolder, NewTargetFolder);
+			}
+
+			string[] Filenames = System.IO.Directory.GetFiles(SourceFolderOnPC);
+			foreach (string Filename in Filenames)
+			{
+				string SourceFilename = Filename;
+				string DestFilename = TargetFolderOnDevice + Path.GetFileName(Filename);
+				WriteProgressLine("Copying '{0}' -> '{1}' ...", 0, SourceFilename, DestFilename);
+				CopyFileToPhone(SourceFilename, DestFilename, 1024 * 1024);
+			}
+		}
+
 		public void DumpInstalledApplications()
 		{
 			Dictionary<string, object> AppBundles;
-			MobileDevice.AMDeviceMethods.LookupApplications(iPhoneHandle, IntPtr.Zero, out AppBundles);
+			MobileDevice.DeviceImpl.LookupApplications(iPhoneHandle, IntPtr.Zero, out AppBundles);
 
 			foreach (var Bundle in AppBundles)
 			{
@@ -539,6 +561,118 @@ namespace Manzana
             }
         }
 
+		/// <summary>
+		/// Tries to copy all of the files in a particular directory on the PC to the phone directory
+		/// (requires the bundle identifier to be able to mount that directory)
+		/// </summary>
+		public bool TryCopy(string BundleIdentifier, string SourceFolderOnPC, string TargetFolderOnDevice)
+		{
+			if (ConnectToBundle(BundleIdentifier))
+			{
+				WriteProgressLine("Connected to bundle '{0}'", 0, BundleIdentifier);
+
+				try
+				{
+					RecursiveCopy(SourceFolderOnPC, TargetFolderOnDevice);
+					return true;
+				}
+				catch (Exception ex)
+				{
+					WriteProgressLine("Failed to transfer a file, extended error is '{0}'", 100, ex.Message);
+					return false;
+				}
+			}
+			else
+			{
+				WriteProgressLine("Error: Failed to connect to bundle '{0}'", 100, BundleIdentifier);
+				return false;
+			}
+		}
+
+		/// <summary>
+		/// Tries to copy all of the files in the manifest on the PC to the documents directory
+		/// (requires the bundle identifier to be able to mount that directory)
+		/// </summary>
+		public bool TryCopy(string BundleIdentifier, string Manifest)
+		{
+			if (ConnectToBundle(BundleIdentifier))
+			{
+				WriteProgressLine("Connected to bundle '{0}'", 0, BundleIdentifier);
+
+				try
+				{
+					string BaseFolder = Path.GetDirectoryName(Manifest);
+					string Files = File.ReadAllText(Manifest);
+					string[] FileList = Files.Split('\n');
+					foreach (string Filename in FileList)
+					{
+						if (!string.IsNullOrEmpty(Filename) && !string.IsNullOrWhiteSpace(Filename))
+						{
+							string Trimmed = Filename.Trim();
+							string SourceFilename = BaseFolder + "\\" + Trimmed;
+							SourceFilename = SourceFilename.Replace('/', '\\');
+							string DestFilename = "/Documents/" + Trimmed.Replace("cookeddata/", "");
+							DestFilename = DestFilename.Replace('\\', '/');
+							WriteProgressLine("Copying '{0}' -> '{1}' ...", 0, SourceFilename, DestFilename);
+							CopyFileToPhone(SourceFilename, DestFilename, 1024 * 1024);
+						}
+					}
+					return true;
+				}
+				catch (Exception ex)
+				{
+					WriteProgressLine("Failed to transfer a file, extended error is '{0}'", 100, ex.Message);
+					return false;
+				}
+			}
+			else
+			{
+				WriteProgressLine("Error: Failed to connect to bundle '{0}'", 100, BundleIdentifier);
+				return false;
+			}
+		}
+
+		/// <summary>
+		/// Tries to copy all of the files in a particular directory on the PC to the phone directory
+		/// (requires the bundle identifier to be able to mount that directory)
+		/// </summary>
+		public bool TryBackup(string BundleIdentifier, string[] Files)
+		{
+			if (ConnectToBundle(BundleIdentifier))
+			{
+				WriteProgressLine("Connected to bundle '{0}'", 0, BundleIdentifier);
+
+				try
+				{
+					string SafeDeviceName = MobileDeviceInstance.SanitizePathNoFilename(DeviceName);
+					foreach (string Filename in Files)
+					{
+						//						string BaseFolder = Path.GetDirectoryName(Filename);
+						string Manifest = Path.GetDirectoryName(Filename) + "\\" + SafeDeviceName + "_" + Path.GetFileName(Filename);
+						CopyFileFromPhone(Manifest, "/Documents/" + Path.GetFileName(Filename), 1024 * 1024);
+					}
+					return true;
+				}
+				catch (Exception ex)
+				{
+					WriteProgressLine("Failed to transfer a file, extended error is '{0}'", 100, ex.Message);
+					return false;
+				}
+			}
+			else
+			{
+				string SafeDeviceName = MobileDeviceInstance.SanitizePathNoFilename(DeviceName);
+				foreach (string Filename in Files)
+				{
+					//						string BaseFolder = Path.GetDirectoryName(Filename);
+					string Manifest = Path.GetDirectoryName (Filename) + "\\" + SafeDeviceName + "_" + Path.GetFileName (Filename);
+					File.WriteAllText (Manifest, "");
+				}
+				WriteProgressLine("Error: Failed to connect to bundle '{0}'", 100, BundleIdentifier);
+				return false;
+			}
+		}
+
         /// <summary>
         /// Returns the names of files in a specified directory
         /// </summary>
@@ -559,14 +693,14 @@ namespace Manzana
             string full_path = FullPath(CurrentDirectory, path);
 
             IntPtr hAFCDir = IntPtr.Zero;
-            if (MobileDevice.AFCDirectoryOpen(AFCCommsHandle, full_path, ref hAFCDir) != 0)
+            if (MobileDevice.DeviceImpl.DirectoryOpen(AFCCommsHandle, full_path, ref hAFCDir) != 0)
             {
                 throw new Exception("Path does not exist");
             }
 
             string buffer = null;
             ArrayList paths = new ArrayList();
-            MobileDevice.AFCDirectoryRead(AFCCommsHandle, hAFCDir, ref buffer);
+			MobileDevice.DeviceImpl.DirectoryRead(AFCCommsHandle, hAFCDir, ref buffer);
 
             while (buffer != null)
             {
@@ -581,9 +715,9 @@ namespace Manzana
                         paths.Add(buffer + "/");
                     }
                 }
-                MobileDevice.AFCDirectoryRead(AFCCommsHandle, hAFCDir, ref buffer);
+				MobileDevice.DeviceImpl.DirectoryRead(AFCCommsHandle, hAFCDir, ref buffer);
             }
-            MobileDevice.AFC.DirectoryClose(AFCCommsHandle, hAFCDir);
+            MobileDevice.DeviceImpl.DirectoryClose(AFCCommsHandle, hAFCDir);
             return (string[])paths.ToArray(typeof(string));
         }
 
@@ -596,13 +730,13 @@ namespace Manzana
             Dictionary<string, string> ans = new Dictionary<string, string>();
             TypedPtr<AFCDictionary> Data;
 
-            int ret = MobileDevice.AFC.FileInfoOpen(AFCCommsHandle, path, out Data);
+			int ret = MobileDevice.DeviceImpl.FileInfoOpen(AFCCommsHandle, path, out Data);
             if ((ret == 0) && (Data.Handle != IntPtr.Zero))
             {
                 IntPtr pname;
                 IntPtr pvalue;
 
-                while (MobileDevice.AFC.KeyValueRead(Data, out pname, out pvalue) == 0 && pname != null && pvalue != null)
+				while (MobileDevice.DeviceImpl.KeyValueRead(Data, out pname, out pvalue) == 0 && pname != IntPtr.Zero && pvalue != IntPtr.Zero)
                 {
                     string name = Marshal.PtrToStringAnsi(pname);
                     string value = Marshal.PtrToStringAnsi(pvalue);
@@ -616,7 +750,7 @@ namespace Manzana
                     }
                 }
 
-                MobileDevice.AFC.KeyValueClose(Data);
+				MobileDevice.DeviceImpl.KeyValueClose(Data);
             }
 
             return ans;
@@ -661,9 +795,9 @@ namespace Manzana
                 // test for symbolic directory link
                 IntPtr hAFCDir = IntPtr.Zero;
 
-                if (directory = (MobileDevice.AFCDirectoryOpen(AFCCommsHandle, path, ref hAFCDir) == 0))
+				if (directory = (MobileDevice.DeviceImpl.DirectoryOpen(AFCCommsHandle, path, ref hAFCDir) == 0))
                 {
-                    MobileDevice.AFC.DirectoryClose(AFCCommsHandle, hAFCDir);
+					MobileDevice.DeviceImpl.DirectoryClose(AFCCommsHandle, hAFCDir);
                 }
             }
         }
@@ -689,7 +823,7 @@ namespace Manzana
         /// <returns>true if directory was created</returns>
         public bool CreateDirectory(string path)
         {
-            return !(MobileDevice.AFC.DirectoryCreate(AFCCommsHandle, FullPath(CurrentDirectory, path)) != 0);
+			return !(MobileDevice.DeviceImpl.DirectoryCreate(AFCCommsHandle, FullPath(CurrentDirectory, path)) != 0);
         }
 
         /// <summary>
@@ -708,7 +842,7 @@ namespace Manzana
             string full_path = FullPath(CurrentDirectory, path);
             //full_path = "/private"; // bug test
 
-            int res = MobileDevice.AFCDirectoryOpen(AFCCommsHandle, full_path, ref hAFCDir);
+			int res = MobileDevice.DeviceImpl.DirectoryOpen(AFCCommsHandle, full_path, ref hAFCDir);
             if (res != 0)
             {
                 throw new Exception("Path does not exist: " + res.ToString());
@@ -716,7 +850,7 @@ namespace Manzana
 
             string buffer = null;
             ArrayList paths = new ArrayList();
-            MobileDevice.AFCDirectoryRead(AFCCommsHandle, hAFCDir, ref buffer);
+			MobileDevice.DeviceImpl.DirectoryRead(AFCCommsHandle, hAFCDir, ref buffer);
 
             while (buffer != null)
             {
@@ -724,9 +858,9 @@ namespace Manzana
                 {
                     paths.Add(buffer);
                 }
-                MobileDevice.AFCDirectoryRead(AFCCommsHandle, hAFCDir, ref buffer);
+				MobileDevice.DeviceImpl.DirectoryRead(AFCCommsHandle, hAFCDir, ref buffer);
             }
-            MobileDevice.AFC.DirectoryClose(AFCCommsHandle, hAFCDir);
+			MobileDevice.DeviceImpl.DirectoryClose(AFCCommsHandle, hAFCDir);
             return (string[])paths.ToArray(typeof(string));
         }
 
@@ -738,7 +872,7 @@ namespace Manzana
         ///	<remarks>Files cannot be moved across filesystem boundaries.</remarks>
         public bool Rename(string sourceName, string destName)
         {
-            return MobileDevice.AFC.RenamePath(AFCCommsHandle, FullPath(CurrentDirectory, sourceName), FullPath(CurrentDirectory, destName)) == 0;
+			return MobileDevice.DeviceImpl.RenamePath(AFCCommsHandle, FullPath(CurrentDirectory, sourceName), FullPath(CurrentDirectory, destName)) == 0;
         }
 
         /// <summary>
@@ -760,10 +894,10 @@ namespace Manzana
         {
             TypedPtr<AFCDictionary> data = IntPtr.Zero;
 
-            int ret = MobileDevice.AFC.FileInfoOpen(AFCCommsHandle, path, out data);
+			int ret = MobileDevice.DeviceImpl.FileInfoOpen(AFCCommsHandle, path, out data);
             if (ret == 0)
             {
-                MobileDevice.AFC.KeyValueClose(data);
+				MobileDevice.DeviceImpl.KeyValueClose(data);
             }
 
             return ret == 0;
@@ -812,7 +946,7 @@ namespace Manzana
             string full_path = FullPath(CurrentDirectory, path);
             if (IsDirectory(full_path))
             {
-                MobileDevice.AFC.RemovePath(AFCCommsHandle, full_path);
+				MobileDevice.DeviceImpl.RemovePath(AFCCommsHandle, full_path);
             }
         }
 
@@ -846,7 +980,7 @@ namespace Manzana
             string full_path = FullPath(CurrentDirectory, path);
             if (Exists(full_path))
             {
-                MobileDevice.AFC.RemovePath(AFCCommsHandle, full_path);
+				MobileDevice.DeviceImpl.RemovePath(AFCCommsHandle, full_path);
             }
         }
         #endregion	// Filesystem
@@ -860,9 +994,9 @@ namespace Manzana
         {
             if (AFCCommsHandle.Handle != IntPtr.Zero)
             {
-                int ans = MobileDevice.AFC.ConnectionClose(AFCCommsHandle);
-                ans = MobileDevice.AMDeviceMethods.StopSession(iPhoneHandle);
-                ans = MobileDevice.AMDeviceMethods.Disconnect(iPhoneHandle);
+				MobileDevice.DeviceImpl.ConnectionClose(AFCCommsHandle);
+				MobileDevice.DeviceImpl.StopSession(iPhoneHandle);
+				MobileDevice.DeviceImpl.Disconnect(iPhoneHandle);
             }
 
             AFCCommsHandle = IntPtr.Zero;
@@ -930,7 +1064,7 @@ namespace Manzana
                 TotalBytesRead += BytesRead;
             }
 
-            DestinationFile.Flush();
+			//            DestinationFile.Flush();
             DestinationFile.Close();
             SourceFile.Close();
 
@@ -942,7 +1076,7 @@ namespace Manzana
         void SetLoggingLevel(int Threshold)
         {
             Int32 LoggingThreshold = Math.Min(Math.Max(0, Threshold), 7);
-            MobileDevice.CFPreferencesSetAppValue(
+            MobileDevice.CoreImpl.CFPreferencesSetAppValue(
                 (IntPtr)MobileDevice.CFStringMakeConstantString("LogLevel"),
                 (IntPtr)MobileDevice.CFNumberCreate(LoggingThreshold),
                 (IntPtr)MobileDevice.CFStringMakeConstantString("com.apple.MobileDevice"));
@@ -952,7 +1086,7 @@ namespace Manzana
         {
             string Line = String.Format(Fmt, Args);
 
-            if (OnGenericProgress != null)
+           if (OnGenericProgress != null)
             {
                 OnGenericProgress(Line, ProgressCount);
             }
@@ -1022,7 +1156,7 @@ namespace Manzana
         void ReconnectWithInstallProxy()
         {
             Reconnect();
-            if (MobileDevice.AMDeviceMethods.StartService(iPhoneHandle, MobileDevice.CFStringMakeConstantString("com.apple.mobile.installation_proxy"), ref hInstallService, IntPtr.Zero) != 0)
+			if (MobileDevice.DeviceImpl.StartService(iPhoneHandle, MobileDevice.CFStringMakeConstantString("com.apple.mobile.installation_proxy"), ref hInstallService, IntPtr.Zero) != 0)
             {
                 Console.WriteLine("Unable to start installation_proxy service!");
             }
@@ -1038,7 +1172,7 @@ namespace Manzana
         /// </summary>
         void HandleProgressCallback(string OuterFunction, TypedPtr<CFDictionary> SourceDict)
         {
-            Dictionary<string, object> Dict = MobileDevice.ConvertCFDictionaryToDictionaryStringy(SourceDict);
+            Dictionary<string, object> Dict = MobileDevice.ConvertCFDictionaryToDictionaryString(SourceDict);
 
             // Expecting:
             // string,string -> "Status",PhaseOfInstaller
@@ -1092,9 +1226,9 @@ namespace Manzana
             IntPtr ClientOptions = IntPtr.Zero;
 
             TypedPtr<CFURL> UrlPath = MobileDevice.CreateFileUrlHelper(IPASourcePathOnPC, false);
-            string UrlPathAsString = MobileDevice.GetStringForUrl(UrlPath);
+			//            string UrlPathAsString = MobileDevice.GetStringForUrl(UrlPath);
 
-            int Result = MobileDevice.AMDeviceMethods.SecureInstallApplication(LiveConnection, iPhoneHandle, UrlPath, ClientOptions, InstallProgressCallback, IntPtr.Zero);
+			int Result = MobileDevice.DeviceImpl.SecureInstallApplication(LiveConnection, iPhoneHandle, UrlPath, ClientOptions, InstallProgressCallback, IntPtr.Zero);
             if (Result == 0)
             {
                 Console.WriteLine("Install of \"{0}\" completed successfully in {2:0.00} seconds", Path.GetFileName(IPASourcePathOnPC), Result, (DateTime.Now - StartTime).TotalSeconds);
@@ -1120,9 +1254,9 @@ namespace Manzana
             IntPtr ClientOptions = IntPtr.Zero;
 
             TypedPtr<CFURL> UrlPath = MobileDevice.CreateFileUrlHelper(IPASourcePathOnPC, false);
-            string UrlPathAsString = MobileDevice.GetStringForUrl(UrlPath);
+			//            string UrlPathAsString = MobileDevice.GetStringForUrl(UrlPath);
 
-            int Result = MobileDevice.AMDeviceMethods.SecureUpgradeApplication(LiveConnection, iPhoneHandle, UrlPath, ClientOptions, InstallProgressCallback, IntPtr.Zero);
+			int Result = MobileDevice.DeviceImpl.SecureUpgradeApplication(LiveConnection, iPhoneHandle, UrlPath, ClientOptions, InstallProgressCallback, IntPtr.Zero);
 
             if (Result == 0)
             {
@@ -1145,11 +1279,11 @@ namespace Manzana
             TypedPtr<CFString> CF_ApplicationIdentifier = MobileDevice.CFStringMakeConstantString(ApplicationIdentifier);
 
             IntPtr CF_ClientOptions = IntPtr.Zero;
-            IntPtr ExtraKey = IntPtr.Zero;
-            IntPtr ExtraValue = IntPtr.Zero;
+			//            IntPtr ExtraKey = IntPtr.Zero;
+			//            IntPtr ExtraValue = IntPtr.Zero;
             IntPtr ConnectionHandle = IntPtr.Zero;
 
-            int Result = MobileDevice.AMDeviceMethods.SecureUninstallApplication(ConnectionHandle, iPhoneHandle, CF_ApplicationIdentifier, CF_ClientOptions, UninstallProgressCallback, IntPtr.Zero);
+			int Result = MobileDevice.DeviceImpl.SecureUninstallApplication(ConnectionHandle, iPhoneHandle, CF_ApplicationIdentifier, CF_ClientOptions, UninstallProgressCallback, IntPtr.Zero);
 
             if (Result == 0)
             {
@@ -1180,35 +1314,35 @@ namespace Manzana
             SetLoggingLevel(7);
 
             // Make sure we can connect to the phone and that it has been paired with this machine previously
-            if (MobileDevice.AMDeviceMethods.Connect(iPhoneHandle) != 0)
+			if (MobileDevice.DeviceImpl.Connect(iPhoneHandle) != 0)
             {
                 return false;
             }
 
-            if (MobileDevice.AMDeviceMethods.IsPaired(iPhoneHandle) != 1)
+			if (MobileDevice.DeviceImpl.IsPaired(iPhoneHandle) != 1)
             {
                 return false;
             }
 
-            if (MobileDevice.AMDeviceMethods.ValidatePairing(iPhoneHandle) != 0)
+			if (MobileDevice.DeviceImpl.ValidatePairing(iPhoneHandle) != 0)
             {
                 return false;
             }
 
             // Start a session
-            if (MobileDevice.AMDeviceMethods.StartSession(iPhoneHandle) == 1)
+			if (MobileDevice.DeviceImpl.StartSession(iPhoneHandle) == 1)
             {
                 return false;
             }
 
-            if (MobileDevice.AMDeviceMethods.StartService(iPhoneHandle, MobileDevice.CFStringMakeConstantString("com.apple.afc"), ref hService, IntPtr.Zero) != 0)
+			if (MobileDevice.DeviceImpl.StartService(iPhoneHandle, MobileDevice.CFStringMakeConstantString("com.apple.afc"), ref hService, IntPtr.Zero) != 0)
             {
                 return false;
             }
 
 
             // Open a file sharing connection
-            if (MobileDevice.AFC.ConnectionOpen(hService, 0, out AFCCommsHandle) != 0)
+			if (MobileDevice.DeviceImpl.ConnectionOpen(hService, 0, out AFCCommsHandle) != 0)
             {
                 return false;
             }
@@ -1224,7 +1358,7 @@ namespace Manzana
             TypedPtr<CFString> CFBundleName = MobileDevice.CFStringMakeConstantString(BundleName);
 
             // Open the bundle
-            int Result = MobileDevice.AMDeviceMethods.StartHouseArrestService(iPhoneHandle, CFBundleName, IntPtr.Zero, ref hService, 0); 
+			int Result = MobileDevice.DeviceImpl.StartHouseArrestService(iPhoneHandle, CFBundleName, IntPtr.Zero, ref hService, 0); 
             if (Result != 0)
             {
                 Console.WriteLine("Failed to connect to bundle '{0}' with {1}", BundleName, MobileDevice.GetErrorString(Result));
@@ -1232,7 +1366,7 @@ namespace Manzana
             }
 
             // Open a file sharing connection
-            if (MobileDevice.AFC.ConnectionOpen(hService, 0, out AFCCommsHandle) != 0)
+			if (MobileDevice.DeviceImpl.ConnectionOpen(hService, 0, out AFCCommsHandle) != 0)
             {
                 return false;
             }
diff --git a/Engine/Source/Programs/IOS/iPhonePackager/Config.cs b/Engine/Source/Programs/IOS/iPhonePackager/Config.cs
--- a/Engine/Source/Programs/IOS/iPhonePackager/Config.cs
+++ b/Engine/Source/Programs/IOS/iPhonePackager/Config.cs
@@ -40,6 +40,15 @@ namespace iPhonePackager
 		public static string RepackageStagingDirectory = "";
 
 		/// <summary>
+		/// The delta manifest for deploying files for iterative deploy
+		/// </summary>
+		public static string DeltaManifest = "";
+
+		/// <summary
+		/// The files to be retrieved from the device
+		/// <summary>
+		public static List<string> FilesForBackup = new List<string>();
+
 		/// The project file that is passed into IPP from UAT
 		/// </summary>
 		public static string ProjectFile = "";
@@ -356,6 +365,11 @@ namespace iPhonePackager
 		/// </summary>
 		public static string MinOSVersion;
 
+		/// <summary>
+		/// Create an iterative IPA (i.e. a stub only with Icons and Splash images)
+		/// </summary>
+		public static bool bIterate = false;
+
 		/// <summary>
 		/// Returns a path to the place to back up documents from a device
 		/// </summary>
diff --git a/Engine/Source/Programs/IOS/iPhonePackager/CookTime.cs b/Engine/Source/Programs/IOS/iPhonePackager/CookTime.cs
--- a/Engine/Source/Programs/IOS/iPhonePackager/CookTime.cs
+++ b/Engine/Source/Programs/IOS/iPhonePackager/CookTime.cs
@@ -304,7 +304,7 @@ namespace iPhonePackager
 			// Add all of the payload files, replacing existing files in the stub IPA if necessary (should only occur for icons)
 			{
 				string SourceDir = Path.GetFullPath(ZipSourceDir);
-				string[] PayloadFiles = Directory.GetFiles(SourceDir, "*.*", SearchOption.AllDirectories);
+				string[] PayloadFiles = Directory.GetFiles(SourceDir, "*.*", Config.bIterate ? SearchOption.TopDirectoryOnly : SearchOption.AllDirectories);
 				foreach (string Filename in PayloadFiles)
 				{
 					// Get the relative path to the file (this implementation only works because we know the files are all
diff --git a/Engine/Source/Programs/IOS/iPhonePackager/DeployTime.cs b/Engine/Source/Programs/IOS/iPhonePackager/DeployTime.cs
--- a/Engine/Source/Programs/IOS/iPhonePackager/DeployTime.cs
+++ b/Engine/Source/Programs/IOS/iPhonePackager/DeployTime.cs
@@ -111,7 +111,15 @@ namespace iPhonePackager
 							ApplicationIdentifier = Utilities.GetStringFromPList("CFBundleIdentifier");
 						}
 
-						if (!DeploymentHelper.Get().BackupDocumentsDirectory(ApplicationIdentifier, Config.GetRootBackedUpDocumentsDirectory()))
+						if (Config.FilesForBackup.Count > 0)
+						{
+							if (!DeploymentHelper.Get().BackupFiles(ApplicationIdentifier, Config.FilesForBackup.ToArray()))
+							{
+								Program.Error("Failed to transfer manifest file from device to PC");
+								Program.ReturnCode = (int)ErrorCodes.Error_DeviceBackupFailed;
+							}
+						}
+						else if (!DeploymentHelper.Get().BackupDocumentsDirectory(ApplicationIdentifier, Config.GetRootBackedUpDocumentsDirectory()))
 						{
 							Program.Error("Failed to transfer documents directory from device to PC");
 							Program.ReturnCode = (int)ErrorCodes.Error_DeviceBackupFailed;
@@ -139,7 +147,7 @@ namespace iPhonePackager
 						string IPAPath = GamePath;
 						string AdditionalCommandline = Program.AdditionalCommandline;
 
-						if (!String.IsNullOrEmpty(AdditionalCommandline))
+						if (!String.IsNullOrEmpty(AdditionalCommandline) && !Config.bIterate)
 						{
 							// Read the mobile provision to check for issues
 							FileOperations.ReadOnlyZipFileSystem Zip = new FileOperations.ReadOnlyZipFileSystem(IPAPath);
@@ -192,7 +200,25 @@ namespace iPhonePackager
 							Zip.Close();
 						}
 
-						if (File.Exists(IPAPath))
+						if (Config.bIterate)
+						{
+							string ApplicationIdentifier = RPCCommand;
+							if (String.IsNullOrEmpty(ApplicationIdentifier))
+							{
+								ApplicationIdentifier = Utilities.GetStringFromPList("CFBundleIdentifier");
+							}
+
+							if (!String.IsNullOrEmpty(Config.DeviceId) && !Config.DeviceId.Contains("All_iOS_On"))
+							{
+								DeploymentHelper.Get().DeviceId = Config.DeviceId;
+							}
+							if (!DeploymentHelper.Get().InstallFilesOnDevice(ApplicationIdentifier, Config.DeltaManifest))
+							{
+								Program.Error("Failed to install Files on device");
+								Program.ReturnCode = (int)ErrorCodes.Error_FilesInstallFailed;
+							}
+						}
+						else if (File.Exists(IPAPath))
 						{
 							if (!String.IsNullOrEmpty(Config.DeviceId) && !Config.DeviceId.Contains("All_iOS_On"))
 							{
@@ -249,9 +275,18 @@ namespace iPhonePackager
 		static Process CreateDeploymentServerProcess()
 		{
 			Process NewProcess = new Process();
-			NewProcess.StartInfo.WorkingDirectory = Path.GetFullPath(".");
-			NewProcess.StartInfo.FileName = NewProcess.StartInfo.WorkingDirectory + "\\DeploymentServer.exe";
-			NewProcess.StartInfo.Arguments = "-iphonepackager " + Process.GetCurrentProcess().Id.ToString();
+			if (Environment.OSVersion.Platform == PlatformID.MacOSX || Environment.OSVersion.Platform == PlatformID.Unix)
+			{
+				NewProcess.StartInfo.WorkingDirectory = Path.GetFullPath(".");
+				NewProcess.StartInfo.FileName = "../../../Build/BatchFiles/Mac/RunMono.sh";
+				NewProcess.StartInfo.Arguments = "\"" + NewProcess.StartInfo.WorkingDirectory + "/DeploymentServer.exe\" -iphonepackager " + Process.GetCurrentProcess().Id.ToString();
+			}
+			else
+			{
+				NewProcess.StartInfo.WorkingDirectory = Path.GetFullPath(".");
+				NewProcess.StartInfo.FileName = NewProcess.StartInfo.WorkingDirectory + "\\DeploymentServer.exe";
+				NewProcess.StartInfo.Arguments = "-iphonepackager " + Process.GetCurrentProcess().Id.ToString();
+			}
 			NewProcess.StartInfo.UseShellExecute = false;
 
 			try
diff --git a/Engine/Source/Programs/IOS/iPhonePackager/iPhonePackager.cs b/Engine/Source/Programs/IOS/iPhonePackager/iPhonePackager.cs
--- a/Engine/Source/Programs/IOS/iPhonePackager/iPhonePackager.cs
+++ b/Engine/Source/Programs/IOS/iPhonePackager/iPhonePackager.cs
@@ -61,6 +61,7 @@ namespace iPhonePackager
         Error_ObbNotFound = 53,
         Error_AndroidBuildToolsPathNotFound = 54,
         Error_NoApkSuitableForArchitecture = 55,
+		Error_FilesInstallFailed = 56,
 		Error_LauncherFailed = 100,
         Error_UATLaunchFailure = 101,
         Error_FailedToDeleteStagingDirectory = 102,
@@ -224,79 +225,154 @@ namespace iPhonePackager
 							case "-cookonthefly":
 								Config.bCookOnTheFly = true;
 								break;
+							case "-iterate":
+								Config.bIterate = true;
+								break;
 						}
 
 						// get the stage dir path
-						if (Arg == "-stagedir") {
+						if (Arg == "-stagedir")
+						{
 							// make sure there's at least one more arg
-							if (Arguments.Length > ArgIndex + 1) {
+							if (Arguments.Length > ArgIndex + 1)
+							{
 								Config.RepackageStagingDirectory = Arguments [++ArgIndex];
-							} else {
+							}
+							else
+							{
+								return false;
+							}
+						}
+						else if (Arg == "-manifest")
+						{
+							// make sure there's at least one more arg
+							if (Arguments.Length > ArgIndex + 1)
+							{
+								Config.DeltaManifest = Arguments[++ArgIndex];
+							}
+							else
+							{
 								return false;
 							}
-						} else if (Arg == "-config") {
+						}
+						else if (Arg == "-backup")
+						{
 							// make sure there's at least one more arg
-							if (Arguments.Length > ArgIndex + 1) {
+							if (Arguments.Length > ArgIndex + 1)
+							{
+								Config.FilesForBackup.Add(Arguments[++ArgIndex]);
+							}
+							else
+							{
+								return false;
+							}
+						}
+						else if (Arg == "-config")
+						{
+							// make sure there's at least one more arg
+							if (Arguments.Length > ArgIndex + 1)
+							{
 								GameConfiguration = Arguments [++ArgIndex];
-							} else {
+							}
+							else
+							{
 								return false;
 							}
 						}
 						// append a name to the bungle identifier and display name
-						else if (Arg == "-bundlename") {
-							if (Arguments.Length > ArgIndex + 1) {
+						else if (Arg == "-bundlename")
+						{
+							if (Arguments.Length > ArgIndex + 1)
+							{
 								Config.OverrideBundleName = Arguments [++ArgIndex];
-							} else {
+							}
+							else
+							{
 								return false;
 							}
-						} else if (Arg == "-mac") {
+						}
+						else if (Arg == "-mac")
+						{
 							// make sure there's at least one more arg
-							if (Arguments.Length > ArgIndex + 1) {
+							if (Arguments.Length > ArgIndex + 1)
+							{
 								Config.OverrideMacName = Arguments [++ArgIndex];
-							} else {
+							}
+							else
+							{
 								return false;
 							}
-						} else if (Arg == "-architecture" || Arg == "-arch") {
+						}
+						else if (Arg == "-architecture" || Arg == "-arch")
+						{
 							// make sure there's at least one more arg
-							if (Arguments.Length > ArgIndex + 1) {
+							if (Arguments.Length > ArgIndex + 1)
+							{
 								Architecture = "-" + Arguments [++ArgIndex];
-							} else {
+							}
+							else
+							{
 								return false;
 							}
-						} else if (Arg == "-project") {
+						}
+						else if (Arg == "-project")
+						{
 							// make sure there's at least one more arg
-							if (Arguments.Length > ArgIndex + 1) {
+							if (Arguments.Length > ArgIndex + 1)
+							{
 								Config.ProjectFile = Arguments [++ArgIndex];
-							} else {
+							}
+							else
+							{
 								return false;
 							}
-						} else if (Arg == "-device") {
+						}
+						else if (Arg == "-device")
+						{
 							// make sure there's at least one more arg
-							if (Arguments.Length > ArgIndex + 1) {
+							if (Arguments.Length > ArgIndex + 1)
+							{
 								Config.DeviceId = Arguments [++ArgIndex];
-							} else {
+							}
+							else
+							{
 								return false;
 							}
-						} else if (Arg == "-additionalcommandline") {
+						}
+						else if (Arg == "-additionalcommandline")
+						{
 							// make sure there's at least one more arg
-							if (Arguments.Length > ArgIndex + 1) {
+							if (Arguments.Length > ArgIndex + 1)
+							{
 								AdditionalCommandline = Arguments [++ArgIndex];
-							} else {
+							}
+							else
+							{
 								return false;
 							}
-						} else if (Arg == "-provision") {
+						}
+						else if (Arg == "-provision")
+						{
 							// make sure there's at least one more arg
-							if (Arguments.Length > ArgIndex + 1) {
+							if (Arguments.Length > ArgIndex + 1)
+							{
 								Config.Provision = Arguments [++ArgIndex];
 								Config.bProvision = true;
-							} else {
+							}
+							else
+							{
 								return false;
 							}
-						} else if (Arg == "-certificate") {
-							if (Arguments.Length > ArgIndex + 1) {
-							Config.Certificate = Arguments [++ArgIndex];
+						}
+						else if (Arg == "-certificate")
+						{
+							if (Arguments.Length > ArgIndex + 1)
+							{
+								Config.Certificate = Arguments [++ArgIndex];
 								Config.bCert = true;
-							} else {
+							}
+							else
+							{
 								return false;
 							}
 						}
diff --git a/Engine/Source/Programs/UnrealBuildTool/IOS/UEDeployIOS.cs b/Engine/Source/Programs/UnrealBuildTool/IOS/UEDeployIOS.cs
--- a/Engine/Source/Programs/UnrealBuildTool/IOS/UEDeployIOS.cs
+++ b/Engine/Source/Programs/UnrealBuildTool/IOS/UEDeployIOS.cs
@@ -443,56 +443,6 @@ namespace UnrealBuildTool.IOS
 						Log.TraceInformation("Copying binaries back to this device failed.");
 					}
 				}
-
-				// install the provision
-/*				string ProvisionWithPrefix = "../../Engine/Build/IOS/UE4Game.mobileprovision";
-				if (File.Exists(BuildPath + "/" + GameName + ".mobileprovision"))
-				{
-					ProvisionWithPrefix = BuildPath + "/" + GameName + ".mobileprovision";
-				}
-				else
-				{
-					if (File.Exists(BuildPath + "/NotForLicensees/" + GameName + ".mobileprovision"))
-					{
-						ProvisionWithPrefix = BuildPath + "/NotForLicensees/" + GameName + ".mobileprovision";
-					}
-					else if (!File.Exists(ProvisionWithPrefix))
-					{
-						ProvisionWithPrefix = "../../Engine/Build/IOS/NotForLicensees/UE4Game.mobileprovision";
-					}
-				}
-				string LibraryDir = Environment.GetFolderPath(Environment.SpecialFolder.LocalApplicationData) + "/Apple Computer/MobileDevice/Provisioning Profiles/";
-				if (File.Exists(ProvisionWithPrefix))
-				{
-					Directory.CreateDirectory(LibraryDir);
-					File.Copy(ProvisionWithPrefix, LibraryDir + GameName + ".mobileprovision", true);
-					FileInfo DestFileInfo = new FileInfo(LibraryDir + GameName + ".mobileprovision");
-					DestFileInfo.Attributes = DestFileInfo.Attributes & ~FileAttributes.ReadOnly;
-				}
-
-				// install the distribution provision
-				ProvisionWithPrefix = "../../Engine/Build/IOS/UE4Game_Distro.mobileprovision";
-				if (File.Exists(BuildPath + "/" + GameName + "_Distro.mobileprovision"))
-				{
-					ProvisionWithPrefix = BuildPath + "/" + GameName + "_Distro.mobileprovision";
-				}
-				else
-				{
-					if (File.Exists(BuildPath + "/NotForLicensees/" + GameName + "_Distro.mobileprovision"))
-					{
-						ProvisionWithPrefix = BuildPath + "/NotForLicensees/" + GameName + "_Distro.mobileprovision";
-					}
-					else if (!File.Exists(ProvisionWithPrefix))
-					{
-						ProvisionWithPrefix = "../../Engine/Build/IOS/NotForLicensees/UE4Game_Distro.mobileprovision";
-					}
-				}
-				if (File.Exists(ProvisionWithPrefix))
-				{
-					File.Copy(ProvisionWithPrefix, LibraryDir + GameName + "_Distro.mobileprovision", true);
-					FileInfo DestFileInfo = new FileInfo(LibraryDir + GameName + "_Distro.mobileprovision");
-					DestFileInfo.Attributes = DestFileInfo.Attributes & ~FileAttributes.ReadOnly;
-				}*/
 			}
 			return true;
 		}
@@ -516,6 +466,12 @@ namespace UnrealBuildTool.IOS
 			}
 
 			SourceFile.CopyTo(DestinationPath, bOverwrite);
+
+			FileInfo DI2 = new FileInfo(DestinationPath);
+			if (DI2.Exists)
+			{
+				DI2.IsReadOnly = false;
+			}
 		}
 
 		private void CopyFiles(string SourceDirectory, string DestinationDirectory, string TargetFiles, bool bOverwrite = false)
diff --git a/Engine/Source/Runtime/Core/Private/IOS/IOSPlatformFile.cpp b/Engine/Source/Runtime/Core/Private/IOS/IOSPlatformFile.cpp
--- a/Engine/Source/Runtime/Core/Private/IOS/IOSPlatformFile.cpp
+++ b/Engine/Source/Runtime/Core/Private/IOS/IOSPlatformFile.cpp
@@ -661,11 +661,17 @@ FString FIOSPlatformFile::ConvertToIOSPath(const FString& Filename, bool bForWri
 		FString Value;
 		// Cache this value as the command line doesn't change...
 		static bool bHasHostIP = FParse::Value(FCommandLine::Get(), TEXT("filehostip"), Value) || FParse::Value(FCommandLine::Get(), TEXT("streaminghostip"), Value);
+		static bool bIsIterative = FParse::Value(FCommandLine::Get(), TEXT("iterative"), Value);
 		if (bHasHostIP)
 		{
 			static FString ReadPathBase = FString([NSSearchPathForDirectoriesInDomains(NSDocumentDirectory, NSUserDomainMask, YES) objectAtIndex:0]) + TEXT("/");
 			return ReadPathBase + Result;
 		}
+		else if (bIsIterative)
+		{
+			static FString ReadPathBase = FString([NSSearchPathForDirectoriesInDomains(NSDocumentDirectory, NSUserDomainMask, YES) objectAtIndex:0]) + TEXT("/");
+			return ReadPathBase + Result.ToLower();
+		}
 		else
 		{
 			static FString ReadPathBase = FString([[NSBundle mainBundle] bundlePath]) + TEXT("/cookeddata/");
