fix for file ref error when deploying on Mac
UE-7300
#ios

diff --git a/Engine/Source/Programs/IOS/MobileDeviceInterface/IPhoneFile.cs b/Engine/Source/Programs/IOS/MobileDeviceInterface/IPhoneFile.cs
--- a/Engine/Source/Programs/IOS/MobileDeviceInterface/IPhoneFile.cs
+++ b/Engine/Source/Programs/IOS/MobileDeviceInterface/IPhoneFile.cs
@@ -267,7 +267,7 @@ namespace Manzana
             }
 
             full_path = phone.FullPath(phone.CurrentDirectory, path);
-			ret = MobileDevice.DeviceImpl.FileRefOpen(phone.AFCHandle, full_path, (int)mode, 0, out handle);
+			ret = MobileDevice.DeviceImpl.FileRefOpen(phone.AFCHandle, full_path, (Int64)mode, out handle);
             if (ret != 0)
             {
                 phone.Reconnect();
diff --git a/Engine/Source/Programs/IOS/MobileDeviceInterface/MobileDevice.cs b/Engine/Source/Programs/IOS/MobileDeviceInterface/MobileDevice.cs
--- a/Engine/Source/Programs/IOS/MobileDeviceInterface/MobileDevice.cs
+++ b/Engine/Source/Programs/IOS/MobileDeviceInterface/MobileDevice.cs
@@ -169,7 +169,7 @@ namespace Manzana
 		int DirectoryRead(TypedPtr<AFCCommConnection> conn, IntPtr dir, ref IntPtr dirent);
 		int FileInfoOpen(TypedPtr<AFCCommConnection> conn, string path, out TypedPtr<AFCDictionary> OutDict);
 		int FileRefClose(TypedPtr<AFCCommConnection> conn, Int64 handle);
-		int FileRefOpen(TypedPtr<AFCCommConnection> conn, string path, int mode, int unknown, out Int64 handle);
+		int FileRefOpen(TypedPtr<AFCCommConnection> conn, string path, Int64 mode, out Int64 handle);
 		int FileRefRead(TypedPtr<AFCCommConnection> conn, Int64 handle, byte[] buffer, ref uint len);
 		int FileRefSeek(TypedPtr<AFCCommConnection> conn, Int64 handle, Int64 pos, Int64 origin);
 		int FileRefSetFileSize(TypedPtr<AFCCommConnection> conn, Int64 handle, uint size);
@@ -364,9 +364,9 @@ namespace Manzana
 			return AFC.FileRefClose(conn, handle);
 		}
 
-		public int FileRefOpen(TypedPtr<AFCCommConnection> conn, string path, int mode, int unknown, out Int64 handle)
+		public int FileRefOpen(TypedPtr<AFCCommConnection> conn, string path, Int64 mode, out Int64 handle)
 		{
-			return AFC.FileRefOpen(conn, path, mode, unknown, out handle);
+			return AFC.FileRefOpen(conn, path, mode, out handle);
 		}
 
 		public int FileRefRead(TypedPtr<AFCCommConnection> conn, Int64 handle, byte[] buffer, ref uint len)
@@ -780,9 +780,9 @@ namespace Manzana
 				return AFCFileRefClose((IntPtr)conn, handle);
 			}
 
-			public static int FileRefOpen(TypedPtr<AFCCommConnection> conn, string path, int mode, int unknown, out Int64 handle)
+			public static int FileRefOpen(TypedPtr<AFCCommConnection> conn, string path, Int64 mode, out Int64 handle)
 			{
-				return AFCFileRefOpen((IntPtr)conn, path, mode, unknown, out handle);
+				return AFCFileRefOpen((IntPtr)conn, path, mode, out handle);
 			}
 
 			public static int FileRefRead(TypedPtr<AFCCommConnection> conn, Int64 handle, byte[] buffer, ref uint len)
@@ -869,7 +869,7 @@ namespace Manzana
 			private extern static int AFCFileRefClose(IntPtr/*AFCCommConnection*/ conn, Int64 handle);
 
 			[DllImport(DLLName, CallingConvention = CallingConvention.Cdecl)]
-			private extern static int AFCFileRefOpen(IntPtr/*AFCCommConnection*/ conn, string path, int mode, int unknown, out Int64 handle);
+			private extern static int AFCFileRefOpen(IntPtr/*AFCCommConnection*/ conn, string path, Int64 mode, out Int64 handle);
 
 			[DllImport(DLLName, CallingConvention = CallingConvention.Cdecl)]
 			private extern static int AFCFileRefRead(IntPtr/*AFCCommConnection*/ conn, Int64 handle, byte[] buffer, ref uint len);
@@ -1132,9 +1132,9 @@ namespace Manzana
 			return AFC.FileRefClose(conn, handle);
 		}
 
-		public int FileRefOpen(TypedPtr<AFCCommConnection> conn, string path, int mode, int unknown, out Int64 handle)
+		public int FileRefOpen(TypedPtr<AFCCommConnection> conn, string path, Int64 mode, out Int64 handle)
 		{
-			return AFC.FileRefOpen(conn, path, mode, unknown, out handle);
+			return AFC.FileRefOpen(conn, path, mode, out handle);
 		}
 
 		public int FileRefRead(TypedPtr<AFCCommConnection> conn, Int64 handle, byte[] buffer, ref uint len)
@@ -1548,9 +1548,9 @@ namespace Manzana
 				return AFCFileRefClose((IntPtr)conn, handle);
 			}
 
-			public static int FileRefOpen(TypedPtr<AFCCommConnection> conn, string path, int mode, int unknown, out Int64 handle)
+			public static int FileRefOpen(TypedPtr<AFCCommConnection> conn, string path, Int64 mode, out Int64 handle)
 			{
-				return AFCFileRefOpen((IntPtr)conn, path, mode, unknown, out handle);
+				return AFCFileRefOpen((IntPtr)conn, path, mode, out handle);
 			}
 
 			public static int FileRefRead(TypedPtr<AFCCommConnection> conn, Int64 handle, byte[] buffer, ref uint len)
@@ -1637,7 +1637,7 @@ namespace Manzana
 			private extern static int AFCFileRefClose(IntPtr/*AFCCommConnection*/ conn, Int64 handle);
 
 			[DllImport(DLLName, CallingConvention = CallingConvention.Cdecl)]
-			private extern static int AFCFileRefOpen(IntPtr/*AFCCommConnection*/ conn, string path, int mode, int unknown, out Int64 handle);
+			private extern static int AFCFileRefOpen(IntPtr/*AFCCommConnection*/ conn, string path, Int64 mode, out Int64 handle);
 
 			[DllImport(DLLName, CallingConvention = CallingConvention.Cdecl)]
 			private extern static int AFCFileRefRead(IntPtr/*AFCCommConnection*/ conn, Int64 handle, byte[] buffer, ref uint len);
