fix issue with trying to install a provision which already exists over itself
UE-7171
#ios

diff --git a/Engine/Source/Programs/IOS/iPhonePackager/ToolsHub.cs b/Engine/Source/Programs/IOS/iPhonePackager/ToolsHub.cs
--- a/Engine/Source/Programs/IOS/iPhonePackager/ToolsHub.cs
+++ b/Engine/Source/Programs/IOS/iPhonePackager/ToolsHub.cs
@@ -172,6 +172,7 @@ namespace iPhonePackager
 					string EffectivePrefix = bIsDistribution ? "Distro_" : Config.SigningPrefix;
 					string DestinationFilename = Path.Combine(Config.ProvisionDirectory, EffectivePrefix + DestName + ".mobileprovision");
 
+					DestinationFilename = DestinationFilename.Replace("\\", "/");
 					if (File.Exists(DestinationFilename))
 					{
 						MobileProvision OldProvision = MobileProvisionParser.ParseFile(DestinationFilename);
@@ -188,10 +189,16 @@ namespace iPhonePackager
 							return;
 						}
 
-						FileOperations.DeleteFile(DestinationFilename);
+						if (DestinationFilename != ProvisionFilename)
+						{
+							FileOperations.DeleteFile(DestinationFilename);
+						}
 					}
 
-					FileOperations.CopyRequiredFile(ProvisionFilename, DestinationFilename);
+					if (DestinationFilename != ProvisionFilename)
+					{
+						FileOperations.CopyRequiredFile(ProvisionFilename, DestinationFilename);
+					}
 				}
 				catch (Exception ex)
 				{
