fix for executable size in the fat binary header when signing
UE-8497
#ios

diff --git a/Engine/Source/Programs/IOS/iPhonePackager/CodeSigning.cs b/Engine/Source/Programs/IOS/iPhonePackager/CodeSigning.cs
--- a/Engine/Source/Programs/IOS/iPhonePackager/CodeSigning.cs
+++ b/Engine/Source/Programs/IOS/iPhonePackager/CodeSigning.cs
@@ -567,13 +567,22 @@ namespace iPhonePackager
 				// Copy the data up to the executable into the final stream
 				if (FatBinary.bIsFatBinary)
 				{
-					OutputExeStream.Seek(0, SeekOrigin.Begin);
-					OutputExeStream.Write(SourceExeData, (int)CurrentStreamOffset, (int)FatBinary.Archs[ArchIndex].Offset - (int)CurrentStreamOffset);
-					OutputExeStream.Seek(0, SeekOrigin.Begin);
+					if (ArchIndex == 0)
+					{
+						OutputExeStream.Seek(0, SeekOrigin.Begin);
+						OutputExeStream.Write(SourceExeData, (int)CurrentStreamOffset, (int)FatBinary.Archs[ArchIndex].Offset - (int)CurrentStreamOffset);
+						OutputExeStream.Seek(0, SeekOrigin.Begin);
 
-					byte[] HeaderData = OutputExeStream.ToArray();
-					HeaderData.CopyTo(FinalExeData, (long)CurrentStreamOffset);
-					CurrentStreamOffset += (ulong)HeaderData.Length;
+						byte[] HeaderData = OutputExeStream.ToArray();
+						HeaderData.CopyTo(FinalExeData, (long)CurrentStreamOffset);
+						CurrentStreamOffset += (ulong)HeaderData.Length;
+					}
+					else
+					{
+						byte[] ZeroData = new byte[(int)FatBinary.Archs[ArchIndex].Offset - (int)CurrentStreamOffset];
+						ZeroData.CopyTo(FinalExeData, (long)CurrentStreamOffset);
+						CurrentStreamOffset += (ulong)ZeroData.Length;
+					}
 				}
 
 				// Copy the executable into the stream
@@ -725,12 +734,21 @@ namespace iPhonePackager
 				}
 				byte[] Data = OutputExeStream.ToArray();
 				Data.CopyTo(FinalExeData, (long)CurrentStreamOffset);
-				CurrentStreamOffset += (ulong)Data.Length;
+				CurrentStreamOffset += DesiredExecutableLength;
+
+				// update the header if it is a fat binary
+				if (FatBinary.bIsFatBinary)
+				{
+					FatBinary.Archs[ArchIndex].Size = (uint)DesiredExecutableLength;
+				}
 
 				// increment the architecture index
 				ArchIndex++;
 			}
 
+			// re-write the header
+			FatBinary.WriteHeader(ref FinalExeData, 0);
+
 			// resize to the finale size
 			Array.Resize(ref FinalExeData, (int)CurrentStreamOffset); //@todo: Extend the file system interface so we don't have to copy 20 MB just to truncate a few hundred bytes
 
diff --git a/Engine/Source/Programs/IOS/iPhonePackager/MachObjects.cs b/Engine/Source/Programs/IOS/iPhonePackager/MachObjects.cs
--- a/Engine/Source/Programs/IOS/iPhonePackager/MachObjects.cs
+++ b/Engine/Source/Programs/IOS/iPhonePackager/MachObjects.cs
@@ -2152,6 +2152,72 @@ namespace MachObjectHandling
 			}
 		}
 
+		public void WriteHeader(ref byte[] OutputData, uint Offset)
+		{
+			if (bIsFatBinary)
+			{
+				// Write the header
+				byte[] Data = BitConverter.GetBytes(Magic);
+				if (BitConverter.IsLittleEndian)
+				{
+					Array.Reverse(Data);
+				}
+				Data.CopyTo(OutputData, Offset);
+				Offset += sizeof(UInt32);
+
+				Data = BitConverter.GetBytes(NumArchs);
+				if (BitConverter.IsLittleEndian)
+				{
+					Array.Reverse(Data);
+				}
+				Data.CopyTo(OutputData, Offset);
+				Offset += sizeof(UInt32);
+
+				foreach (FatBinaryArch Arch in Archs)
+				{
+					Data = BitConverter.GetBytes(Arch.CpuType);
+					if (BitConverter.IsLittleEndian)
+					{
+						Array.Reverse(Data);
+					}
+					Data.CopyTo(OutputData, Offset);
+					Offset += sizeof(UInt32);
+
+					Data = BitConverter.GetBytes(Arch.CpuSubType);
+					if (BitConverter.IsLittleEndian)
+					{
+						Array.Reverse(Data);
+					}
+					Data.CopyTo(OutputData, Offset);
+					Offset += sizeof(UInt32);
+
+					Data = BitConverter.GetBytes(Arch.Offset);
+					if (BitConverter.IsLittleEndian)
+					{
+						Array.Reverse(Data);
+					}
+					Data.CopyTo(OutputData, Offset);
+					Offset += sizeof(UInt32);
+
+					Data = BitConverter.GetBytes(Arch.Size);
+					if (BitConverter.IsLittleEndian)
+					{
+						Array.Reverse(Data);
+					}
+					Data.CopyTo(OutputData, Offset);
+					Offset += sizeof(UInt32);
+
+					Data = BitConverter.GetBytes(Arch.Align);
+					if (BitConverter.IsLittleEndian)
+					{
+						Array.Reverse(Data);
+					}
+					Data.CopyTo(OutputData, Offset);
+					Offset += sizeof(UInt32);
+				}
+			}
+		}
+
 		public void LoadFromFile(string Filename)
 		{
 			BinaryReader SR = new BinaryReader(File.OpenRead(Filename));
