fixed Deployment Server to properly return an error code if it fails
UE-8692
#ios

diff --git a/Engine/Source/Programs/IOS/DeploymentServer/Program.cs b/Engine/Source/Programs/IOS/DeploymentServer/Program.cs
--- a/Engine/Source/Programs/IOS/DeploymentServer/Program.cs
+++ b/Engine/Source/Programs/IOS/DeploymentServer/Program.cs
@@ -34,7 +34,9 @@ namespace DeploymentServer
 
     class Program
     {
-        static void Main(string[] args)
+		static int ExitCode = 0;
+
+        static int Main(string[] args)
         {
             if ((args.Length == 2) && (args[0].Equals("-iphonepackager")))
             {
@@ -61,10 +63,6 @@ namespace DeploymentServer
             }
             else
             {
-				// Run directly by some intrepid explorer
-				Console.WriteLine ("Note: This program should only be started by iPhonePackager");
-				Console.WriteLine ("  This program cannot be used on it's own.");
-
 				// Parse the command
 				if (ParseCommand(args))
 				{
@@ -90,6 +88,9 @@ namespace DeploymentServer
 				}
                 Console.WriteLine("Exiting.");
             }
+
+			Environment.ExitCode = Program.ExitCode;
+			return Program.ExitCode;
         }
 
 		private static string Command = "";
@@ -178,20 +179,24 @@ namespace DeploymentServer
 			{
 				Deployer.DeviceId = Device;
 			}
+
+			bool bResult = true;
 			switch (Command)
 			{
 				case "backup":
-					Deployer.BackupFiles(Bundle, FileList.ToArray());
+					bResult = Deployer.BackupFiles(Bundle, FileList.ToArray());
 					break;
 
 				case "deploy":
-					Deployer.InstallFilesOnDevice(Bundle, Manifest);
+					bResult = Deployer.InstallFilesOnDevice(Bundle, Manifest);
 					break;
 
 				case "install":
-					Deployer.InstallIPAOnDevice(ipaPath);
+					bResult = Deployer.InstallIPAOnDevice(ipaPath);
 					break;
 			}
+
+			Program.ExitCode = bResult ? 0 : 1;
 		}
     }
 }
