proper order for picking mobile provisions for code signing
UE-8733
#ios

diff --git a/Engine/Source/Programs/IOS/iPhonePackager/CodeSigning.cs b/Engine/Source/Programs/IOS/iPhonePackager/CodeSigning.cs
--- a/Engine/Source/Programs/IOS/iPhonePackager/CodeSigning.cs
+++ b/Engine/Source/Programs/IOS/iPhonePackager/CodeSigning.cs
@@ -302,7 +302,7 @@ namespace iPhonePackager
 			string SelectedProvision = "";
 			string SelectedCert = "";
 			string SelectedFile = "";
-			bool bFoundName = false;
+			int FoundName = -1;
 			Dictionary<string, MobileProvision> ProvisionLibrary = new Dictionary<string, MobileProvision>();
 			foreach (string Provision in Directory.EnumerateFiles(Config.ProvisionDirectory, "*.mobileprovision"))
 			{
@@ -320,10 +320,24 @@ namespace iPhonePackager
 					bValid = (Cert.NotBefore < Now) && (Cert.NotAfter > Now);
 				}
 				bool bPassesNameCheck = p.ApplicationIdentifier.Substring(p.ApplicationIdentifierPrefix.Length+1) == CFBundleIdentifier;
-				bool bPassesWildCardCheck = p.ApplicationIdentifier.Contains("*");
+				bool bPassesCompanyCheck = false;
+				bool bPassesWildCardCheck = false;
+				if (p.ApplicationIdentifier.Contains("*"))
+				{
+					string CompanyName = p.ApplicationIdentifier.Substring(p.ApplicationIdentifierPrefix.Length + 1);
+					if (CompanyName != "*")
+					{
+						CompanyName = CompanyName.Substring(0, CompanyName.LastIndexOf("."));
+						bPassesCompanyCheck = CFBundleIdentifier.StartsWith(CompanyName);
+					}
+					else
+					{
+						bPassesWildCardCheck = true;
+					}
+				}
 				bool bDistribution = ((p.ProvisionedDeviceIDs.Count == 0) && !p.bDebug);
 				string Validity = "VALID";
-				if (!bPassesNameCheck && !bPassesWildCardCheck)
+				if (!bPassesNameCheck && !bPassesWildCardCheck && !bPassesCompanyCheck)
 				{
 					Validity = "NO_MATCH";
 				}
@@ -338,12 +352,27 @@ namespace iPhonePackager
 						Validity = "EXPIRED";
 					}
 				}
-				if ((string.IsNullOrWhiteSpace(SelectedProvision) || !bFoundName) && Validity == "VALID" && !bDistribution)
+				if ((string.IsNullOrWhiteSpace(SelectedProvision) || FoundName < 2) && Validity == "VALID" && !bDistribution)
 				{
-					bFoundName = bPassesNameCheck;
-					SelectedProvision = p.ProvisionName;
-					SelectedFile = Path.GetFileName(Provision);
-					SelectedCert = Cert.FriendlyName;
+					int Prev = FoundName;
+					if (bPassesNameCheck)
+					{
+						FoundName = 2;
+					}
+					else if (bPassesCompanyCheck && FoundName < 1)
+					{
+						FoundName = 1;
+					}
+					else if (bPassesWildCardCheck && FoundName == -1)
+					{
+						FoundName = 0;
+					}
+					if (FoundName != Prev)
+					{
+						SelectedProvision = p.ProvisionName;
+						SelectedFile = Path.GetFileName(Provision);
+						SelectedCert = Cert.FriendlyName;
+					}
 				}
 				Program.LogVerbose("PROVISION-File:{0},Name:{1},Validity:{2},StartDate:{3},EndDate:{4},Type:{5}", Path.GetFileName(Provision), p.ProvisionName, Validity, EffectiveDate.ToString(), ExpirationDate.ToString(), bDistribution ? "DISTRIBUTION" : "DEVELOPMENT");
 			}
diff --git a/Engine/Source/Programs/IOS/iPhonePackager/MobileProvisionUtilities.cs b/Engine/Source/Programs/IOS/iPhonePackager/MobileProvisionUtilities.cs
--- a/Engine/Source/Programs/IOS/iPhonePackager/MobileProvisionUtilities.cs
+++ b/Engine/Source/Programs/IOS/iPhonePackager/MobileProvisionUtilities.cs
@@ -99,7 +99,7 @@ namespace iPhonePackager
 
 			// check the cache for a provision matching the app id (com.company.Game)
 			// First checking for a contains match and then for a wildcard match
-			for (int Phase = 0; Phase < 2; ++Phase)
+			for (int Phase = 0; Phase < 3; ++Phase)
 			{
 				foreach (KeyValuePair<string, MobileProvision> Pair in ProvisionLibrary)
 				{
@@ -112,14 +112,25 @@ namespace iPhonePackager
 					bool bPassesNameCheck = false;
 					if (Phase == 0)
 					{
-						bPassesNameCheck = TestProvision.ApplicationIdentifier.Contains(CFBundleIdentifier);
+						bPassesNameCheck = TestProvision.ApplicationIdentifier.Substring(TestProvision.ApplicationIdentifierPrefix.Length+1) == CFBundleIdentifier;
 						bNameMatch = bPassesNameCheck;
 					}
+					else if (Phase == 1)
+					{
+						if (TestProvision.ApplicationIdentifier.Contains("*"))
+						{
+							string CompanyName = TestProvision.ApplicationIdentifier.Substring(TestProvision.ApplicationIdentifierPrefix.Length + 1);
+							if (CompanyName != "*")
+							{
+								CompanyName = CompanyName.Substring(0, CompanyName.LastIndexOf("."));
+								bPassesNameCheck = CFBundleIdentifier.StartsWith(CompanyName);
+							}
+						}
+					}
 					else
 					{
-						bPassesNameCheck = TestProvision.ProvisionName.Contains("Wildcard") || TestProvision.ApplicationIdentifier.Contains("*");
+						bPassesNameCheck = TestProvision.ApplicationIdentifier.Contains("*");
 					}
-
 					if (!bPassesNameCheck && bCheckIdentifier)
 					{
 						Program.LogVerbose("  .. Failed phase {0} name check (provision app ID was {1})", Phase, TestProvision.ApplicationIdentifier);
