fix for change in iTunes Mobile Device location with latest iTunes
UE-8861
#ios

diff --git a/Engine/Source/Developer/iOS/IOSTargetPlatform/Private/IOSTargetPlatform.cpp b/Engine/Source/Developer/iOS/IOSTargetPlatform/Private/IOSTargetPlatform.cpp
--- a/Engine/Source/Developer/iOS/IOSTargetPlatform/Private/IOSTargetPlatform.cpp
+++ b/Engine/Source/Developer/iOS/IOSTargetPlatform/Private/IOSTargetPlatform.cpp
@@ -90,8 +90,11 @@ bool FIOSTargetPlatform::IsSdkInstalled(bool bProjectHasCode, FString& OutTutori
 		unsigned long pathSize = 256;
 		if (lRes != ERROR_SUCCESS || RegQueryValueEx(hKey, TEXT("iTunesMobileDeviceDLL"), 0, NULL, (BYTE*)dllPath, &pathSize) != ERROR_SUCCESS || IFileManager::Get().FileSize(*FString(dllPath)) == INDEX_NONE)
 		{
-			OutTutorialPath = FString("/Engine/Tutorial/Mobile/InstallingiTunesTutorial.InstallingiTunesTutorial");
-			biOSSDKInstalled = false;
+			if (RegQueryValueEx(hKey, TEXT("MobileDeviceDLL"), 0, NULL, (BYTE*)dllPath, &pathSize) != ERROR_SUCCESS || IFileManager::Get().FileSize(*FString(dllPath)) == INDEX_NONE)
+			{
+				OutTutorialPath = FString("/Engine/Tutorial/Mobile/InstallingiTunesTutorial.InstallingiTunesTutorial");
+				biOSSDKInstalled = false;
+			}
 		}
 	}
 
diff --git a/Engine/Source/Programs/IOS/MobileDeviceInterface/MobileDevice.cs b/Engine/Source/Programs/IOS/MobileDeviceInterface/MobileDevice.cs
--- a/Engine/Source/Programs/IOS/MobileDeviceInterface/MobileDevice.cs
+++ b/Engine/Source/Programs/IOS/MobileDeviceInterface/MobileDevice.cs
@@ -907,12 +907,12 @@ namespace Manzana
 		#endregion
 	}
 
-	internal class MobileDeviceWin : MobileDeviceImpl
+	internal class MobileDeviceWiniTunes11 : MobileDeviceImpl
 	{
 		//		static readonly int ForceStaticInit = 42;
 		const string DLLName = "iTunesMobileDevice.dll";
 
-		static MobileDeviceWin()
+		static MobileDeviceWiniTunes11()
 		{
 			List<string> PathBits = new List<string>();
 			PathBits.Add(Environment.GetEnvironmentVariable("Path"));
@@ -1675,6 +1675,774 @@ namespace Manzana
 		#endregion
 	}
 
+	internal class MobileDeviceWiniTunes12 : MobileDeviceImpl
+	{
+		//		static readonly int ForceStaticInit = 42;
+		const string DLLName = "MobileDevice.dll";
+
+		static MobileDeviceWiniTunes12()
+		{
+			List<string> PathBits = new List<string>();
+			PathBits.Add(Environment.GetEnvironmentVariable("Path"));
+
+			// Try to add the paths from the registry (they aren't always available on newer iTunes installs though)
+			object RegistryDllPath = Registry.GetValue(@"HKEY_LOCAL_MACHINE\SOFTWARE\Apple Inc.\Apple Mobile Device Support\Shared", "MobileDeviceDLL", DLLName);
+			if (RegistryDllPath != null)
+			{
+				FileInfo iTunesMobileDeviceFile = new FileInfo(RegistryDllPath.ToString());
+
+				if (iTunesMobileDeviceFile.Exists)
+				{
+					PathBits.Add(iTunesMobileDeviceFile.DirectoryName);
+				}
+			}
+
+			object RegistrySupportDir = Registry.GetValue(@"HKEY_LOCAL_MACHINE\SOFTWARE\Apple Inc.\Apple Application Support", "InstallDir", Environment.CurrentDirectory);
+			if (RegistrySupportDir != null)
+			{
+				DirectoryInfo ApplicationSupportDirectory = new DirectoryInfo(RegistrySupportDir.ToString());
+
+				if (ApplicationSupportDirectory.Exists)
+				{
+					PathBits.Add(ApplicationSupportDirectory.FullName);
+				}
+			}
+
+			// Add some guesses as well
+			DirectoryInfo AppleMobileDeviceSupport = new DirectoryInfo(Environment.GetFolderPath(Environment.SpecialFolder.CommonProgramFiles) + @"\Apple\Mobile Device Support");
+			if (AppleMobileDeviceSupport.Exists)
+			{
+				PathBits.Add(AppleMobileDeviceSupport.FullName);
+			}
+
+			DirectoryInfo AppleMobileDeviceSupportX86 = new DirectoryInfo(Environment.GetFolderPath(Environment.SpecialFolder.CommonProgramFilesX86) + @"\Apple\Mobile Device Support");
+			if ((AppleMobileDeviceSupport != AppleMobileDeviceSupportX86) && (AppleMobileDeviceSupportX86.Exists))
+			{
+				PathBits.Add(AppleMobileDeviceSupportX86.FullName);
+			}
+
+			// Set the path from all the individual bits
+			Environment.SetEnvironmentVariable("Path", string.Join(";", PathBits));
+		}
+
+		public int NotificationSubscribe(DeviceNotificationCallback DeviceCallbackHandle)
+		{
+			return AMDeviceMethods.NotificationSubscribe(DeviceCallbackHandle);
+		}
+
+		public IntPtr CopyDeviceIdentifier(TypedPtr<AppleMobileDeviceConnection> device)
+		{
+			return AMDeviceMethods.CopyDeviceIdentifier(device);
+		}
+
+		public int Connect(TypedPtr<AppleMobileDeviceConnection> device)
+		{
+			return AMDeviceMethods.Connect(device);
+		}
+
+		public IntPtr CopyValue(TypedPtr<AppleMobileDeviceConnection> device, uint unknown, TypedPtr<CFString> cfstring)
+		{
+			return AMDeviceMethods.CopyValue(device, unknown, cfstring);
+		}
+
+		public int Disconnect(TypedPtr<AppleMobileDeviceConnection> device)
+		{
+			return AMDeviceMethods.Disconnect(device);
+		}
+
+		public int GetConnectionID(TypedPtr<AppleMobileDeviceConnection> device)
+		{
+			return AMDeviceMethods.GetConnectionID(device);
+		}
+
+		public int IsPaired(TypedPtr<AppleMobileDeviceConnection> device)
+		{
+			return AMDeviceMethods.IsPaired(device);
+		}
+
+		public int Pair(TypedPtr<AppleMobileDeviceConnection> device)
+		{
+			return AMDeviceMethods.Pair(device);
+		}
+
+		public int Unpair(TypedPtr<AppleMobileDeviceConnection> device)
+		{
+			return AMDeviceMethods.Unpair(device);
+		}
+
+		public int ValidatePairing(TypedPtr<AppleMobileDeviceConnection> device)
+		{
+			return AMDeviceMethods.ValidatePairing(device);
+		}
+
+		public int LookupApplications(TypedPtr<AppleMobileDeviceConnection> device, IntPtr options, out Dictionary<string, object> AppBundles)
+		{
+			int Result = AMDeviceMethods.LookupApplications(device, options, out AppBundles);
+			return Result;
+		}
+
+		public int StartSession(TypedPtr<AppleMobileDeviceConnection> device)
+		{
+			return AMDeviceMethods.StartSession((IntPtr)device);
+		}
+
+		public int StopSession(TypedPtr<AppleMobileDeviceConnection> device)
+		{
+			return AMDeviceMethods.StopSession((IntPtr)device);
+		}
+
+		public int StartHouseArrestService(TypedPtr<AppleMobileDeviceConnection> device, TypedPtr<CFString> bundleName, IntPtr unknown1, ref IntPtr handle, int unknown2)
+		{
+			return AMDeviceMethods.StartHouseArrestService((IntPtr)device, (IntPtr)bundleName, unknown1, ref handle, unknown2);
+		}
+
+		public int StartService(TypedPtr<AppleMobileDeviceConnection> device, TypedPtr<CFString> serviceName, ref IntPtr handle, IntPtr unknown)
+		{
+			return AMDeviceMethods.StartService((IntPtr)device, (IntPtr)serviceName, ref handle, unknown);
+		}
+
+		public int RestoreRegisterForDeviceNotifications(
+			DeviceRestoreNotificationCallback dfu_connect,
+			DeviceRestoreNotificationCallback recovery_connect,
+			DeviceRestoreNotificationCallback dfu_disconnect,
+			DeviceRestoreNotificationCallback recovery_disconnect,
+			uint unknown0,
+			IntPtr user_info)
+		{
+			return AMDeviceMethods.AMRestoreRegisterForDeviceNotifications(dfu_connect, recovery_connect, dfu_disconnect, recovery_disconnect, unknown0, user_info);
+		}
+
+		public int SecureUninstallApplication(
+			IntPtr serviceConnection,
+			TypedPtr<AppleMobileDeviceConnection> DeviceIfConnIsNull,
+			TypedPtr<CFString> ApplicationIdentifer,
+			TypedPtr<CFDictionary> ClientOptions,
+			DeviceInstallationCallback ProgressCallback,
+			IntPtr UserData)
+		{
+			return AMDeviceMethods.SecureUninstallApplication(serviceConnection, DeviceIfConnIsNull, ApplicationIdentifer, ClientOptions, ProgressCallback, UserData);
+		}
+
+		public int SecureInstallApplication(
+			IntPtr serviceConnection,
+			TypedPtr<AppleMobileDeviceConnection> DeviceIfConnIsNull,
+			TypedPtr<CFURL> UrlPath,
+			TypedPtr<CFDictionary> ClientOptions,
+			DeviceInstallationCallback ProgressCallback,
+			IntPtr UserData)
+		{
+			return AMDeviceMethods.SecureInstallApplication(serviceConnection, DeviceIfConnIsNull, UrlPath, ClientOptions, ProgressCallback, UserData);
+		}
+
+		public int SecureUpgradeApplication(
+			IntPtr ServiceConnection,
+			TypedPtr<AppleMobileDeviceConnection> DeviceIfConnIsNull,
+			TypedPtr<CFURL> UrlPath,
+			TypedPtr<CFDictionary> ClientOptions,
+			DeviceInstallationCallback ProgressCallback,
+			IntPtr UserData)
+		{
+			return AMDeviceMethods.SecureUpgradeApplication(ServiceConnection, DeviceIfConnIsNull, UrlPath, ClientOptions, ProgressCallback, UserData);
+		}
+
+		public int ConnectionClose(TypedPtr<AFCCommConnection> conn)
+		{
+			return AFC.ConnectionClose(conn);
+		}
+
+		public int ConnectionInvalidate(TypedPtr<AFCCommConnection> conn)
+		{
+			return AFC.ConnectionInvalidate(conn);
+		}
+
+		public int ConnectionIsValid(TypedPtr<AFCCommConnection> conn)
+		{
+			return AFC.ConnectionIsValid(conn);
+		}
+
+		public int ConnectionOpen(IntPtr handle, uint io_timeout, out TypedPtr<AFCCommConnection> OutConn)
+		{
+			return AFC.ConnectionOpen(handle, io_timeout, out OutConn);
+		}
+
+		public int DeviceInfoOpen(IntPtr handle, ref IntPtr dict)
+		{
+			return AFC.DeviceInfoOpen(handle, ref dict);
+		}
+
+		public int DirectoryClose(TypedPtr<AFCCommConnection> conn, IntPtr dir)
+		{
+			return AFC.DirectoryClose(conn, dir);
+		}
+
+		public int DirectoryCreate(TypedPtr<AFCCommConnection> conn, string path)
+		{
+			return AFC.DirectoryCreate(conn, path);
+		}
+
+		public int DirectoryOpen(TypedPtr<AFCCommConnection> conn, byte[] path, ref IntPtr dir)
+		{
+			return AFC.DirectoryOpen(conn, path, ref dir);
+		}
+
+		public int DirectoryRead(TypedPtr<AFCCommConnection> conn, IntPtr dir, ref IntPtr dirent)
+		{
+			return AFC.DirectoryRead(conn, dir, ref dirent);
+		}
+
+		public int FileInfoOpen(TypedPtr<AFCCommConnection> conn, string path, out TypedPtr<AFCDictionary> OutDict)
+		{
+			return AFC.FileInfoOpen(conn, path, out OutDict);
+		}
+
+		public int FileRefClose(TypedPtr<AFCCommConnection> conn, Int64 handle)
+		{
+			return AFC.FileRefClose(conn, handle);
+		}
+
+		public int FileRefOpen(TypedPtr<AFCCommConnection> conn, string path, Int64 mode, out Int64 handle)
+		{
+			return AFC.FileRefOpen(conn, path, mode, out handle);
+		}
+
+		public int FileRefRead(TypedPtr<AFCCommConnection> conn, Int64 handle, byte[] buffer, ref uint len)
+		{
+			return AFC.FileRefRead(conn, handle, buffer, ref len);
+		}
+
+		public int FileRefSeek(TypedPtr<AFCCommConnection> conn, Int64 handle, Int64 pos, Int64 origin)
+		{
+			return AFC.FileRefSeek(conn, handle, pos, origin);
+		}
+
+		public int FileRefSetFileSize(TypedPtr<AFCCommConnection> conn, Int64 handle, uint size)
+		{
+			return AFC.FileRefSetFileSize(conn, handle, size);
+		}
+
+		public int FileRefTell(TypedPtr<AFCCommConnection> conn, Int64 handle, ref uint position)
+		{
+			return AFC.FileRefTell(conn, handle, ref position);
+		}
+
+		public int FileRefWrite(TypedPtr<AFCCommConnection> conn, Int64 handle, byte[] buffer, uint len)
+		{
+			return AFC.FileRefWrite(conn, handle, buffer, len);
+		}
+
+		public int FlushData(TypedPtr<AFCCommConnection> conn, Int64 handle)
+		{
+			return AFC.FlushData(conn, handle);
+		}
+
+		public int KeyValueClose(TypedPtr<AFCDictionary> dict)
+		{
+			return AFC.KeyValueClose(dict);
+		}
+
+		public int KeyValueRead(TypedPtr<AFCDictionary> dict, out IntPtr key, out IntPtr val)
+		{
+			return AFC.KeyValueRead(dict, out key, out val);
+		}
+
+		public int RemovePath(TypedPtr<AFCCommConnection> conn, string path)
+		{
+			return AFC.RemovePath(conn, path);
+		}
+
+		public int RenamePath(TypedPtr<AFCCommConnection> conn, string OldPath, string NewPath)
+		{
+			return AFC.RenamePath(conn, OldPath, NewPath);
+		}
+
+		#region Application Methods
+		public class AMDeviceMethods
+		{
+			// Need to have the path manipulation code in static MobileDevice() to run
+			//			static int ForceMobileDeviceConstructorToRun;
+			static AMDeviceMethods()
+			{
+				//				ForceMobileDeviceConstructorToRun = ForceStaticInit;
+			}
+
+			public static IntPtr CopyDeviceIdentifier(TypedPtr<AppleMobileDeviceConnection> device)
+			{
+				return AMDeviceCopyDeviceIdentifier((IntPtr)device);
+			}
+
+			public static int Connect(TypedPtr<AppleMobileDeviceConnection> device)
+			{
+				return AMDeviceConnect((IntPtr)device);
+			}
+
+			public static IntPtr CopyValue(TypedPtr<AppleMobileDeviceConnection> device, uint unknown, TypedPtr<CFString> cfstring)
+			{
+				return AMDeviceCopyValue((IntPtr)device, unknown, (IntPtr)cfstring);
+			}
+
+			public static int Disconnect(TypedPtr<AppleMobileDeviceConnection> device)
+			{
+				return AMDeviceDisconnect((IntPtr)device);
+			}
+
+			public static int GetConnectionID(TypedPtr<AppleMobileDeviceConnection> device)
+			{
+				return AMDeviceGetConnectionID((IntPtr)device);
+			}
+
+			public static int IsPaired(TypedPtr<AppleMobileDeviceConnection> device)
+			{
+				return AMDeviceIsPaired((IntPtr)device);
+			}
+
+			public static int Pair(TypedPtr<AppleMobileDeviceConnection> device)
+			{
+				return AMDevicePair((IntPtr)device);
+			}
+
+			public static int Unpair(TypedPtr<AppleMobileDeviceConnection> device)
+			{
+				return AMDeviceUnpair((IntPtr)device);
+			}
+
+			public static int ValidatePairing(TypedPtr<AppleMobileDeviceConnection> device)
+			{
+				return AMDeviceValidatePairing((IntPtr)device);
+			}
+
+			public static int LookupApplications(TypedPtr<AppleMobileDeviceConnection> device, IntPtr options, out Dictionary<string, object> AppBundles)
+			{
+				IntPtr UntypedDict = IntPtr.Zero;
+				int Result = AMDeviceLookupApplications((IntPtr)device, options, ref UntypedDict);
+
+				AppBundles = MobileDevice.ConvertCFDictionaryToDictionaryString(new TypedPtr<CFDictionary>(UntypedDict));
+
+				return Result;
+			}
+
+			public static int StartSession(TypedPtr<AppleMobileDeviceConnection> device)
+			{
+				return AMDeviceStartSession((IntPtr)device);
+			}
+
+			public static int StopSession(TypedPtr<AppleMobileDeviceConnection> device)
+			{
+				return AMDeviceStopSession((IntPtr)device);
+			}
+
+			public static int StartHouseArrestService(TypedPtr<AppleMobileDeviceConnection> device, TypedPtr<CFString> bundleName, IntPtr unknown1, ref IntPtr handle, int unknown2)
+			{
+				return AMDeviceStartHouseArrestService((IntPtr)device, (IntPtr)bundleName, unknown1, ref handle, unknown2);
+			}
+
+			public static int StartService(TypedPtr<AppleMobileDeviceConnection> device, TypedPtr<CFString> serviceName, ref IntPtr handle, IntPtr unknown)
+			{
+				return AMDeviceStartService((IntPtr)device, (IntPtr)serviceName, ref handle, unknown);
+			}
+
+			public static int UninstallApplication(
+				TypedPtr<AppleMobileDeviceConnection> device,
+				TypedPtr<CFString> ApplicationIdentifier,
+				TypedPtr<CFDictionary> ClientOptions,
+				DeviceInstallationCallback ProgressCallback,
+				IntPtr UserData)
+			{
+				return AMDeviceUninstallApplication((IntPtr)device, (IntPtr)ApplicationIdentifier, (IntPtr)ClientOptions, ProgressCallback, UserData);
+			}
+
+			public static int SecureUninstallApplication(
+				IntPtr serviceConnection,
+				TypedPtr<AppleMobileDeviceConnection> DeviceIfConnIsNull,
+				TypedPtr<CFString> ApplicationIdentifer,
+				TypedPtr<CFDictionary> ClientOptions,
+				DeviceInstallationCallback ProgressCallback,
+				IntPtr UserData)
+			{
+				return AMDeviceSecureUninstallApplication(serviceConnection, (IntPtr)DeviceIfConnIsNull, (IntPtr)ApplicationIdentifer, (IntPtr)ClientOptions, ProgressCallback, UserData);
+			}
+
+			public static int SecureInstallApplication(
+				IntPtr serviceConnection,
+				TypedPtr<AppleMobileDeviceConnection> DeviceIfConnIsNull,
+				TypedPtr<CFURL> UrlPath,
+				TypedPtr<CFDictionary> ClientOptions,
+				DeviceInstallationCallback ProgressCallback,
+				IntPtr UserData)
+			{
+				return AMDeviceSecureInstallApplication(serviceConnection, (IntPtr)DeviceIfConnIsNull, (IntPtr)UrlPath, (IntPtr)ClientOptions, ProgressCallback, UserData);
+			}
+
+			public static int TransferApplication(
+				TypedPtr<AppleMobileDeviceConnection> device,
+				TypedPtr<CFString> InPackagePath,
+				IntPtr UnknownButUnused,
+				DeviceInstallationCallback ProgressCallback,
+				IntPtr UserData)
+			{
+				return AMDeviceTransferApplication((IntPtr)device, (IntPtr)InPackagePath, UnknownButUnused, ProgressCallback, UserData);
+			}
+
+			public static int SecureUpgradeApplication(
+				IntPtr ServiceConnection,
+				TypedPtr<AppleMobileDeviceConnection> DeviceIfConnIsNull,
+				TypedPtr<CFURL> UrlPath,
+				TypedPtr<CFDictionary> ClientOptions,
+				DeviceInstallationCallback ProgressCallback,
+				IntPtr UserData)
+			{
+				return AMDeviceSecureUpgradeApplication(ServiceConnection, (IntPtr)DeviceIfConnIsNull, (IntPtr)UrlPath, (IntPtr)ClientOptions, ProgressCallback, UserData);
+			}
+
+			public static int InstallApplication(
+				TypedPtr<AppleMobileDeviceConnection> device,
+				TypedPtr<CFString> FilePath,
+				TypedPtr<CFDictionary> ClientOptions,
+				DeviceInstallationCallback ProgressCallback,
+				IntPtr UserData)
+			{
+				return AMDeviceInstallApplication((IntPtr)device, (IntPtr)FilePath, (IntPtr)ClientOptions, ProgressCallback, UserData);
+			}
+
+			public static int NotificationSubscribe(DeviceNotificationCallback DeviceCallbackHandle)
+			{
+				IntPtr notification;
+				return AMDeviceNotificationSubscribe(DeviceCallbackHandle, 0, 0, 0, out notification);
+			}
+
+			[DllImport(DLLName, CallingConvention = CallingConvention.Cdecl)]
+			private extern static IntPtr AMDeviceCopyDeviceIdentifier(IntPtr/*AppleMobileDeviceConnection*/ device);
+
+			[DllImport(DLLName, CallingConvention = CallingConvention.Cdecl)]
+			private extern static int AMDeviceConnect(IntPtr/*AppleMobileDeviceConnection*/ device);
+
+			[DllImport(DLLName, CallingConvention = CallingConvention.Cdecl)]
+			private extern static IntPtr AMDeviceCopyValue(IntPtr/*AppleMobileDeviceConnection*/ device, uint unknown, IntPtr/*CFString*/ cfstring);
+
+			[DllImport(DLLName, CallingConvention = CallingConvention.Cdecl)]
+			private extern static int AMDeviceDisconnect(IntPtr/*AppleMobileDeviceConnection*/ device);
+
+			[DllImport(DLLName, CallingConvention = CallingConvention.Cdecl)]
+			private extern static int AMDeviceGetConnectionID(IntPtr/*AppleMobileDeviceConnection*/ device);
+
+			[DllImport(DLLName, CallingConvention = CallingConvention.Cdecl)]
+			private extern static int AMDeviceIsPaired(IntPtr/*AppleMobileDeviceConnection*/ device);
+
+			[DllImport(DLLName, CallingConvention = CallingConvention.Cdecl)]
+			private extern static int AMDevicePair(IntPtr/*AppleMobileDeviceConnection*/ device);
+
+			[DllImport(DLLName, CallingConvention = CallingConvention.Cdecl)]
+			private extern static int AMDeviceUnpair(IntPtr/*AppleMobileDeviceConnection*/ device);
+
+			[DllImport(DLLName, CallingConvention = CallingConvention.Cdecl)]
+			private extern static int AMDeviceValidatePairing(IntPtr/*AppleMobileDeviceConnection*/ device);
+
+			[DllImport(DLLName, CallingConvention = CallingConvention.Cdecl)]
+			private extern static int AMDeviceLookupApplications(IntPtr/*AppleMobileDeviceConnection*/ device, IntPtr options, ref IntPtr appBundles);
+
+			[DllImport(DLLName, CallingConvention = CallingConvention.Cdecl)]
+			private extern static int AMDeviceNotificationSubscribe(DeviceNotificationCallback callback, uint unused1, uint unused2, uint unused3, out IntPtr am_device_notification_ptr);
+
+			[DllImport(DLLName, CallingConvention = CallingConvention.Cdecl)]
+			private extern static int AMDeviceStartHouseArrestService(IntPtr/*AppleMobileDeviceConnection*/ device, IntPtr/*CFString*/ bundleName, IntPtr unknown1, ref IntPtr handle, int unknown2);
+
+			[DllImport(DLLName, CallingConvention = CallingConvention.Cdecl)]
+			private extern static int AMDeviceStartService(IntPtr/*AppleMobileDeviceConnection*/ device, IntPtr/*CFString*/ serviceName, ref IntPtr handle, IntPtr unknown);
+
+			[DllImport(DLLName, CallingConvention = CallingConvention.Cdecl)]
+			private extern static int AMDeviceStartSession(IntPtr/*AppleMobileDeviceConnection*/ device);
+
+			[DllImport(DLLName, CallingConvention = CallingConvention.Cdecl)]
+			private extern static int AMDeviceStopSession(IntPtr/*AppleMobileDeviceConnection*/ device);
+
+			[DllImport(DLLName, CallingConvention = CallingConvention.Cdecl)]
+			private extern static int AMRestoreModeDeviceCreate(uint unknown0, int connection_id, uint unknown1);
+
+			[DllImport(DLLName, CallingConvention = CallingConvention.Cdecl)]
+			public extern static int AMRestoreRegisterForDeviceNotifications(
+				DeviceRestoreNotificationCallback dfu_connect,
+				DeviceRestoreNotificationCallback recovery_connect,
+				DeviceRestoreNotificationCallback dfu_disconnect,
+				DeviceRestoreNotificationCallback recovery_disconnect,
+				uint unknown0,
+				IntPtr user_info);
+
+			[DllImport(DLLName, CallingConvention = CallingConvention.Cdecl)]
+			private extern static int AMDeviceUninstallApplication(
+				IntPtr/*AppleMobileDeviceConnection*/ device,
+				IntPtr/*CFString*/ ApplicationIdentifier,
+				IntPtr/*CFDictionary*/ ClientOptions,
+				DeviceInstallationCallback ProgressCallback,
+				IntPtr UserData);
+
+			[DllImport(DLLName, CallingConvention = CallingConvention.Cdecl)]
+			private extern static int AMDeviceSecureUninstallApplication(
+				IntPtr serviceConnection,
+				IntPtr/*AppleMobileDeviceConnection*/ DeviceIfConnIsNull,
+				IntPtr/*CFString*/ ApplicationIdentifer,
+				IntPtr/*CFDictionary*/ ClientOptions,
+				DeviceInstallationCallback ProgressCallback,
+				IntPtr UserData);
+
+			[DllImport(DLLName, CallingConvention = CallingConvention.Cdecl)]
+			private extern static int AMDeviceSecureInstallApplication(
+				IntPtr serviceConnection,
+				IntPtr/*AppleMobileDeviceConnection*/ DeviceIfConnIsNull,
+				IntPtr/*CFURL*/ UrlPath,
+				IntPtr/*CFDictionary*/ ClientOptions,
+				DeviceInstallationCallback ProgressCallback,
+				IntPtr UserData);
+
+			[DllImport(DLLName, CallingConvention = CallingConvention.Cdecl)]
+			private extern static int AMDeviceTransferApplication(
+				IntPtr/*AppleMobileDeviceConnection*/ device,
+				IntPtr/*CFString*/ InPackagePath,
+				IntPtr UnknownButUnused,
+				DeviceInstallationCallback ProgressCallback,
+				IntPtr UserData);
+
+			[DllImport(DLLName, CallingConvention = CallingConvention.Cdecl)]
+			private extern static int AMDeviceSecureUpgradeApplication(
+				IntPtr ServiceConnection,
+				IntPtr/*AppleMobileDeviceConnection*/ DeviceIfConnIsNull,
+				IntPtr/*CFURL*/ UrlPath,
+				IntPtr/*CFDictionary*/ ClientOptions,
+				DeviceInstallationCallback ProgressCallback,
+				IntPtr UserData);
+
+			[DllImport(DLLName, CallingConvention = CallingConvention.Cdecl)]
+			private extern static int AMDeviceInstallApplication(
+				IntPtr/*AppleMobileDeviceConnection*/ device,
+				IntPtr/*CFString*/ FilePath,
+				IntPtr/*CFDictionary*/ ClientOptions,
+				DeviceInstallationCallback ProgressCallback,
+				IntPtr UserData);
+
+			// the unknown goes into the 4th dword in the connection struct.  if non-zero, winsock send() won't get called in AMDServiceConnectionSend AFAIK
+			// the only option is "CloseOnInvalidate" (CFBoolean type, defaults to false)
+			[DllImport(DLLName, CallingConvention = CallingConvention.Cdecl)]
+			private extern static IntPtr AMDServiceConnectionCreate(IntPtr/*AppleMobileDeviceConnection*/ device, IntPtr UnknownTypically0, IntPtr OptionsDictionary);
+
+
+			[DllImport(DLLName, CallingConvention = CallingConvention.Cdecl)]
+			private extern static int AMSBeginSync(IntPtr/*AppleMobileDeviceConnection*/ device, IntPtr OuterIn8, IntPtr OuterIn12, IntPtr OuterIn16);
+		}
+		#endregion
+
+		#region AFC Operations
+
+		public int DirectoryOpen(TypedPtr<AFCCommConnection> conn, string path, ref IntPtr dir)
+		{
+			return AFC.DirectoryOpen((IntPtr)conn, Encoding.UTF8.GetBytes(path), ref dir);
+		}
+
+		public int DirectoryRead(TypedPtr<AFCCommConnection> conn, IntPtr dir, ref string buffer)
+		{
+			int ret;
+
+			IntPtr ptr = IntPtr.Zero;
+			ret = AFC.DirectoryRead((IntPtr)conn, dir, ref ptr);
+			if ((ret == 0) && (ptr != IntPtr.Zero))
+			{
+				buffer = Marshal.PtrToStringAnsi(ptr);
+			}
+			else
+			{
+				buffer = null;
+			}
+			return ret;
+		}
+
+		public class AFC
+		{
+			public static int ConnectionClose(TypedPtr<AFCCommConnection> conn)
+			{
+				return AFCConnectionClose((IntPtr)conn);
+			}
+
+			public static int ConnectionInvalidate(TypedPtr<AFCCommConnection> conn)
+			{
+				return AFCConnectionInvalidate((IntPtr)conn);
+			}
+
+			public static int ConnectionIsValid(TypedPtr<AFCCommConnection> conn)
+			{
+				return AFCConnectionIsValid((IntPtr)conn);
+			}
+
+			public static int ConnectionOpen(IntPtr handle, uint io_timeout, out TypedPtr<AFCCommConnection> OutConn)
+			{
+				IntPtr Conn;
+				int Result = AFCConnectionOpen(handle, io_timeout, out Conn);
+				OutConn = Conn;
+				return Result;
+			}
+
+			public static int DeviceInfoOpen(IntPtr handle, ref IntPtr dict)
+			{
+				return AFCDeviceInfoOpen(handle, ref dict);
+			}
+
+			public static int DirectoryClose(TypedPtr<AFCCommConnection> conn, IntPtr dir)
+			{
+				return AFCDirectoryClose((IntPtr)conn, dir);
+			}
+
+			public static int DirectoryCreate(TypedPtr<AFCCommConnection> conn, string path)
+			{
+				return AFCDirectoryCreate((IntPtr)conn, path);
+			}
+
+			public static int DirectoryOpen(TypedPtr<AFCCommConnection> conn, byte[] path, ref IntPtr dir)
+			{
+				return AFCDirectoryOpen((IntPtr)conn, path, ref dir);
+			}
+
+			public static int DirectoryRead(TypedPtr<AFCCommConnection> conn, IntPtr dir, ref IntPtr dirent)
+			{
+				return AFCDirectoryRead((IntPtr)conn, dir, ref dirent);
+			}
+
+			public static int FileInfoOpen(TypedPtr<AFCCommConnection> conn, string path, out TypedPtr<AFCDictionary> OutDict)
+			{
+				IntPtr UntypedDict;
+				int Result = AFCFileInfoOpen((IntPtr)conn, path, out UntypedDict);
+				OutDict = UntypedDict;
+
+				return Result;
+			}
+
+			public static int FileRefClose(TypedPtr<AFCCommConnection> conn, Int64 handle)
+			{
+				return AFCFileRefClose((IntPtr)conn, handle);
+			}
+
+			public static int FileRefOpen(TypedPtr<AFCCommConnection> conn, string path, Int64 mode, out Int64 handle)
+			{
+				return AFCFileRefOpen((IntPtr)conn, path, mode, out handle);
+			}
+
+			public static int FileRefRead(TypedPtr<AFCCommConnection> conn, Int64 handle, byte[] buffer, ref uint len)
+			{
+				return AFCFileRefRead((IntPtr)conn, handle, buffer, ref len);
+			}
+
+			public static int FileRefSetFileSize(TypedPtr<AFCCommConnection> conn, Int64 handle, uint size)
+			{
+				return AFCFileRefSetFileSize((IntPtr)conn, handle, size);
+			}
+
+			public static int FileRefSeek(TypedPtr<AFCCommConnection> conn, Int64 handle, Int64 pos, Int64 origin)
+			{
+				return AFCFileRefSeek((IntPtr)conn, handle, pos, origin);
+			}
+
+			public static int FileRefTell(TypedPtr<AFCCommConnection> conn, Int64 handle, ref uint position)
+			{
+				return AFCFileRefTell((IntPtr)conn, handle, ref position);
+			}
+
+			public static int FileRefWrite(TypedPtr<AFCCommConnection> conn, Int64 handle, byte[] buffer, uint len)
+			{
+				return AFCFileRefWrite((IntPtr)conn, handle, buffer, len);
+			}
+
+			public static int FlushData(TypedPtr<AFCCommConnection> conn, Int64 handle)
+			{
+				return AFCFlushData((IntPtr)conn, handle);
+			}
+
+			public static int KeyValueClose(TypedPtr<AFCDictionary> dict)
+			{
+				return AFCKeyValueClose((IntPtr)dict);
+			}
+
+			public static int KeyValueRead(TypedPtr<AFCDictionary> dict, out IntPtr key, out IntPtr val)
+			{
+				return AFCKeyValueRead((IntPtr)dict, out key, out val);
+			}
+
+			public static int RemovePath(TypedPtr<AFCCommConnection> conn, string path)
+			{
+				return AFCRemovePath((IntPtr)conn, path);
+			}
+
+			public static int RenamePath(TypedPtr<AFCCommConnection> conn, string OldPath, string NewPath)
+			{
+				return AFCRenamePath((IntPtr)conn, OldPath, NewPath);
+			}
+
+			[DllImport(DLLName, CallingConvention = CallingConvention.Cdecl)]
+			private extern static int AFCConnectionClose(IntPtr/*AFCCommConnection*/ conn);
+
+			[DllImport(DLLName, CallingConvention = CallingConvention.Cdecl)]
+			private extern static int AFCConnectionInvalidate(IntPtr/*AFCCommConnection*/ conn);
+
+			[DllImport(DLLName, CallingConvention = CallingConvention.Cdecl)]
+			private extern static int AFCConnectionIsValid(IntPtr/*AFCCommConnection*/ conn);
+
+			[DllImport(DLLName, CallingConvention = CallingConvention.Cdecl)]
+			private extern static int AFCConnectionOpen(IntPtr handle, uint io_timeout, out IntPtr/*AFCCommConnection*/ conn);
+
+			[DllImport(DLLName, CallingConvention = CallingConvention.Cdecl)]
+			private extern static int AFCDeviceInfoOpen(IntPtr handle, ref IntPtr dict);
+
+			[DllImport(DLLName, CallingConvention = CallingConvention.Cdecl)]
+			private extern static int AFCDirectoryClose(IntPtr/*AFCCommConnection*/ conn, IntPtr dir);
+
+			[DllImport(DLLName, CallingConvention = CallingConvention.Cdecl)]
+			private extern static int AFCDirectoryCreate(IntPtr/*AFCCommConnection*/ conn, string path);
+
+			[DllImport(DLLName, CallingConvention = CallingConvention.Cdecl)]
+			private extern static int AFCDirectoryOpen(IntPtr/*AFCCommConnection*/ conn, byte[] path, ref IntPtr dir);
+
+			[DllImport(DLLName, CallingConvention = CallingConvention.Cdecl)]
+			private extern static int AFCDirectoryRead(IntPtr/*AFCCommConnection*/ conn, IntPtr dir, ref IntPtr dirent);
+
+			[DllImport(DLLName, CallingConvention = CallingConvention.Cdecl)]
+			private extern static int AFCFileInfoOpen(IntPtr/*AFCCommConnection*/ conn, string path, out IntPtr OutDict);
+
+			[DllImport(DLLName, CallingConvention = CallingConvention.Cdecl)]
+			private extern static int AFCFileRefClose(IntPtr/*AFCCommConnection*/ conn, Int64 handle);
+
+			[DllImport(DLLName, CallingConvention = CallingConvention.Cdecl)]
+			private extern static int AFCFileRefOpen(IntPtr/*AFCCommConnection*/ conn, string path, Int64 mode, out Int64 handle);
+
+			[DllImport(DLLName, CallingConvention = CallingConvention.Cdecl)]
+			private extern static int AFCFileRefRead(IntPtr/*AFCCommConnection*/ conn, Int64 handle, byte[] buffer, ref uint len);
+
+			[DllImport(DLLName, CallingConvention = CallingConvention.Cdecl)]
+			private extern static int AFCFileRefSeek(IntPtr/*AFCCommConnection*/ conn, Int64 handle, Int64 pos, Int64 origin);
+
+			// FIXME - not working, arguments?
+			[DllImport(DLLName, CallingConvention = CallingConvention.Cdecl)]
+			public extern static int AFCFileRefSetFileSize(IntPtr/*AFCCommConnection*/ conn, Int64 handle, uint size);
+
+			[DllImport(DLLName, CallingConvention = CallingConvention.Cdecl)]
+			private extern static int AFCFileRefTell(IntPtr/*AFCCommConnection*/ conn, Int64 handle, ref uint position);
+
+			[DllImport(DLLName, CallingConvention = CallingConvention.Cdecl)]
+			private extern static int AFCFileRefWrite(IntPtr/*AFCCommConnection*/ conn, Int64 handle, byte[] buffer, uint len);
+
+			[DllImport(DLLName, CallingConvention = CallingConvention.Cdecl)]
+			private extern static int AFCFlushData(IntPtr/*AFCCommConnection*/ conn, Int64 handle);
+
+			[DllImport(DLLName, CallingConvention = CallingConvention.Cdecl)]
+			private extern static int AFCKeyValueClose(IntPtr dict);
+
+			[DllImport(DLLName, CallingConvention = CallingConvention.Cdecl)]
+			private extern static int AFCKeyValueRead(IntPtr dict, out IntPtr key, out IntPtr val);
+
+
+			[DllImport(DLLName, CallingConvention = CallingConvention.Cdecl)]
+			private extern static int AFCRemovePath(IntPtr/*AFCCommConnection*/ conn, string path);
+
+			[DllImport(DLLName, CallingConvention = CallingConvention.Cdecl)]
+			private extern static int AFCRenamePath(IntPtr/*AFCCommConnection*/ conn, string old_path, string new_path);
+		}
+
+		#endregion
+	}
+
     internal partial class MobileDevice
     {
 		static public MobileDeviceImpl DeviceImpl = null;
@@ -1687,7 +2455,16 @@ namespace Manzana
 			}
 			else
 			{
-				DeviceImpl = new MobileDeviceWin();
+				string dllPath11 = Microsoft.Win32.Registry.GetValue("HKEY_LOCAL_MACHINE\\SOFTWARE\\Wow6432Node\\Apple Inc.\\Apple Mobile Device Support\\Shared", "iTunesMobileDeviceDLL", null) as string;
+				string dllPath12 = Microsoft.Win32.Registry.GetValue("HKEY_LOCAL_MACHINE\\SOFTWARE\\Wow6432Node\\Apple Inc.\\Apple Mobile Device Support\\Shared", "MobileDeviceDLL", null) as string;
+				if (String.IsNullOrEmpty(dllPath11) || !File.Exists(dllPath11))
+				{
+					DeviceImpl = new MobileDeviceWiniTunes11();
+				}
+				else if (String.IsNullOrEmpty(dllPath12) || !File.Exists(dllPath12))
+				{
+					DeviceImpl = new MobileDeviceWiniTunes12();
+				}
 			}
 
             // Initialize the CoreFoundation bindings
diff --git a/Engine/Source/Programs/IOS/iPhonePackager/iPhonePackager.cs b/Engine/Source/Programs/IOS/iPhonePackager/iPhonePackager.cs
--- a/Engine/Source/Programs/IOS/iPhonePackager/iPhonePackager.cs
+++ b/Engine/Source/Programs/IOS/iPhonePackager/iPhonePackager.cs
@@ -678,7 +678,11 @@ namespace iPhonePackager
 						}
 						else
 						{
-							dllPath = Microsoft.Win32.Registry.GetValue("HKEY_LOCAL_MACHINE\\SOFTWARE\\Wow6432Node\\Apple Inc.\\Apple Mobile Device Support\\Shared", "iTunesMobileDeviceDLL", null) as string; 
+							dllPath = Microsoft.Win32.Registry.GetValue("HKEY_LOCAL_MACHINE\\SOFTWARE\\Wow6432Node\\Apple Inc.\\Apple Mobile Device Support\\Shared", "iTunesMobileDeviceDLL", null) as string;
+							if (String.IsNullOrEmpty(dllPath) || !File.Exists(dllPath))
+							{
+								dllPath = Microsoft.Win32.Registry.GetValue("HKEY_LOCAL_MACHINE\\SOFTWARE\\Wow6432Node\\Apple Inc.\\Apple Mobile Device Support\\Shared", "MobileDeviceDLL", null) as string;
+							}
 						}
 						if (String.IsNullOrEmpty(dllPath) || !File.Exists(dllPath))
 						{
