fix for latest iTunes not being compatible with the deployment server
UE-8861
#ios

diff --git a/Engine/Source/Programs/IOS/DeploymentServer/Program.cs b/Engine/Source/Programs/IOS/DeploymentServer/Program.cs
--- a/Engine/Source/Programs/IOS/DeploymentServer/Program.cs
+++ b/Engine/Source/Programs/IOS/DeploymentServer/Program.cs
@@ -14,6 +14,49 @@ namespace DeploymentServer
 {
 	internal class CoreFoundation
 	{
+		static public CoreFoundationImpl CoreImpl;
+
+		static CoreFoundation()
+		{
+			if (Environment.OSVersion.Platform == PlatformID.MacOSX || Environment.OSVersion.Platform == PlatformID.Unix)
+			{
+				CoreImpl = new CoreFoundationOSX();
+			}
+			else
+			{
+				CoreImpl = new CoreFoundationWin32();
+			}
+		}
+
+		public static int RunLoopRunInMode(IntPtr mode, double seconds, int returnAfterSourceHandled)
+		{
+			return CoreImpl.RunLoopInMode(mode, seconds, returnAfterSourceHandled);
+		}
+
+		public static IntPtr kCFRunLoopDefaultMode()
+		{
+			return CoreImpl.DefaultMode();
+		}
+	};
+
+	internal interface CoreFoundationImpl
+	{
+		int RunLoopInMode(IntPtr mode, double seconds, int returnAfterSourceHandled);
+		IntPtr DefaultMode();
+	}
+
+	internal class CoreFoundationOSX : CoreFoundationImpl
+	{
+		public int RunLoopInMode(IntPtr mode, double seconds, int returnAfterSourceHandled)
+		{
+			return CFRunLoopRunInMode(mode, seconds, returnAfterSourceHandled);
+		}
+
+		public IntPtr DefaultMode()
+		{
+			return kCFRunLoopDefaultMode;
+		}
+
 		[DllImport("/System/Library/Frameworks/CoreFoundation.framework/CoreFoundation")]
 		public extern static IntPtr CFStringCreateWithCString(IntPtr allocator, string value, int encoding);
 
@@ -32,7 +75,37 @@ namespace DeploymentServer
 		public static IntPtr kCFRunLoopDefaultMode = CFStringCreateWithCString(IntPtr.Zero, "kCFRunLoopDefaultMode", 0);
 	}
 
-    class Program
+	internal class CoreFoundationWin32 : CoreFoundationImpl
+	{
+		public int RunLoopInMode(IntPtr mode, double seconds, int returnAfterSourceHandled)
+		{
+			return CFRunLoopRunInMode(mode, seconds, returnAfterSourceHandled);
+		}
+
+		public IntPtr DefaultMode()
+		{
+			return kCFRunLoopDefaultMode;
+		}
+
+		[DllImport("CoreFoundation.dll")]
+		public extern static IntPtr CFStringCreateWithCString(IntPtr allocator, string value, int encoding);
+
+		[DllImport("CoreFoundation.dll")]
+		public extern static void CFRunLoopRun();
+
+		[DllImport("CoreFoundation.dll")]
+		public extern static IntPtr CFRunLoopGetMain();
+
+		[DllImport("CoreFoundation.dll")]
+		public extern static IntPtr CFRunLoopGetCurrent();
+
+		[DllImport("CoreFoundation.dll")]
+		public extern static int CFRunLoopRunInMode(IntPtr mode, double seconds, int returnAfterSourceHandled);
+
+		public static IntPtr kCFRunLoopDefaultMode = CFStringCreateWithCString(IntPtr.Zero, "kCFRunLoopDefaultMode", 0);
+	}
+
+	class Program
     {
 		static int ExitCode = 0;
 
@@ -68,22 +141,15 @@ namespace DeploymentServer
 				{
 					Deployer = new DeploymentImplementation();
 					bool bCommandComplete = false;
-					if (Environment.OSVersion.Platform == PlatformID.MacOSX || Environment.OSVersion.Platform == PlatformID.Unix)
-					{
-						System.Threading.Thread enumerateLoop = new System.Threading.Thread(delegate()
-						{
-							RunCommand();
-							bCommandComplete = true;
-						});
-						enumerateLoop.Start();
-						while (!bCommandComplete)
-						{
-							CoreFoundation.CFRunLoopRunInMode(CoreFoundation.kCFRunLoopDefaultMode, 1.0, 0);
-						}
-					}
-					else
+					System.Threading.Thread enumerateLoop = new System.Threading.Thread(delegate()
 					{
 						RunCommand();
+						bCommandComplete = true;
+					});
+					enumerateLoop.Start();
+					while (!bCommandComplete)
+					{
+						CoreFoundation.RunLoopRunInMode(CoreFoundation.kCFRunLoopDefaultMode(), 1.0, 0);
 					}
 				}
                 Console.WriteLine("Exiting.");
@@ -194,6 +260,10 @@ namespace DeploymentServer
 				case "install":
 					bResult = Deployer.InstallIPAOnDevice(ipaPath);
 					break;
+
+				case "enumerate":
+					Deployer.EnumerateConnectedDevices();
+					break;
 			}
 
 			Program.ExitCode = bResult ? 0 : 1;
diff --git a/Engine/Source/Programs/IOS/MobileDeviceInterface/MobileDevice.cs b/Engine/Source/Programs/IOS/MobileDeviceInterface/MobileDevice.cs
--- a/Engine/Source/Programs/IOS/MobileDeviceInterface/MobileDevice.cs
+++ b/Engine/Source/Programs/IOS/MobileDeviceInterface/MobileDevice.cs
@@ -1678,7 +1678,7 @@ namespace Manzana
 	internal class MobileDeviceWiniTunes12 : MobileDeviceImpl
 	{
 		//		static readonly int ForceStaticInit = 42;
-		const string DLLName = "MobileDevice.dll";
+		const string DLLName = @"MobileDevice.dll";
 
 		static MobileDeviceWiniTunes12()
 		{
@@ -2457,11 +2457,11 @@ namespace Manzana
 			{
 				string dllPath11 = Microsoft.Win32.Registry.GetValue("HKEY_LOCAL_MACHINE\\SOFTWARE\\Wow6432Node\\Apple Inc.\\Apple Mobile Device Support\\Shared", "iTunesMobileDeviceDLL", null) as string;
 				string dllPath12 = Microsoft.Win32.Registry.GetValue("HKEY_LOCAL_MACHINE\\SOFTWARE\\Wow6432Node\\Apple Inc.\\Apple Mobile Device Support\\Shared", "MobileDeviceDLL", null) as string;
-				if (String.IsNullOrEmpty(dllPath11) || !File.Exists(dllPath11))
+				if (!String.IsNullOrEmpty(dllPath11) && File.Exists(dllPath11))
 				{
 					DeviceImpl = new MobileDeviceWiniTunes11();
 				}
-				else if (String.IsNullOrEmpty(dllPath12) || !File.Exists(dllPath12))
+				else if (!String.IsNullOrEmpty(dllPath12) && File.Exists(dllPath12))
 				{
 					DeviceImpl = new MobileDeviceWiniTunes12();
 				}
diff --git a/Engine/Source/Programs/UnrealBuildTool/IOS/UEBuildIOS.cs b/Engine/Source/Programs/UnrealBuildTool/IOS/UEBuildIOS.cs
--- a/Engine/Source/Programs/UnrealBuildTool/IOS/UEBuildIOS.cs
+++ b/Engine/Source/Programs/UnrealBuildTool/IOS/UEBuildIOS.cs
@@ -29,7 +29,11 @@ namespace UnrealBuildTool
 				string dllPath = Microsoft.Win32.Registry.GetValue("HKEY_LOCAL_MACHINE\\SOFTWARE\\Wow6432Node\\Apple Inc.\\Apple Mobile Device Support\\Shared", "iTunesMobileDeviceDLL", null) as string;
 				if (String.IsNullOrEmpty(dllPath) || !File.Exists(dllPath))
 				{
-					return SDKStatus.Invalid;
+					dllPath = Microsoft.Win32.Registry.GetValue("HKEY_LOCAL_MACHINE\\SOFTWARE\\Wow6432Node\\Apple Inc.\\Apple Mobile Device Support\\Shared", "MobileDeviceDLL", null) as string;
+					if (String.IsNullOrEmpty(dllPath) || !File.Exists(dllPath))
+					{
+						return SDKStatus.Invalid;
+					}
 				}
 			}
 			return SDKStatus.Valid;
