fix for standalone IPP not being able to deploy or even see devices
#ios

diff --git a/Engine/Source/Programs/IOS/DeploymentServer/DeployTime.cs b/Engine/Source/Programs/IOS/DeploymentServer/DeployTime.cs
--- a/Engine/Source/Programs/IOS/DeploymentServer/DeployTime.cs
+++ b/Engine/Source/Programs/IOS/DeploymentServer/DeployTime.cs
@@ -43,8 +43,58 @@ namespace DeploymentServer
             return 25;
         }
     }
-    
-    class DeploymentImplementation : MarshalByRefObject, DeploymentInterface
+
+	class DeploymentProxy : MarshalByRefObject, DeploymentInterface
+	{
+		static public DeploymentImplementation Deployer;
+		public DeploymentProxy()
+		{
+		}
+
+		public string DeviceId { get { return Deployer.DeviceId; } set { Deployer.DeviceId = value; } }
+
+		public void SetReportingInterface(DeployTimeReportingInterface InReporter)
+		{
+			Deployer.SetReportingInterface(InReporter);
+		}
+
+		public bool UninstallIPAOnDevice(string ApplicationIdentifier)
+		{
+			return Deployer.UninstallIPAOnDevice(ApplicationIdentifier);
+		}
+
+		public ConnectedDeviceInfo[] EnumerateConnectedDevices()
+		{
+			return Deployer.EnumerateConnectedDevices();
+		}
+
+		public bool InstallIPAOnDevice(string IPAPath)
+		{
+			return Deployer.InstallIPAOnDevice(IPAPath);
+		}
+
+		public bool InstallFilesOnDevice(string BundleIdentifier, string ManifestFile)
+		{
+			return Deployer.InstallFilesOnDevice(BundleIdentifier, ManifestFile);
+		}
+
+		public bool BackupDocumentsDirectory(string BundleIdentifier, string DestinationDocumentsDirectory)
+		{
+			return Deployer.BackupDocumentsDirectory(BundleIdentifier, DestinationDocumentsDirectory);
+		}
+
+		public bool BackupFiles(string BundleIdentifier, string[] Files)
+		{
+			return Deployer.BackupFiles(BundleIdentifier, Files);
+		}
+
+		public bool ListApplications()
+		{
+			return Deployer.ListApplications();
+		}
+	};
+
+    class DeploymentImplementation : DeploymentInterface
     {
         bool bHaveRegisteredHandlers = false;
 
diff --git a/Engine/Source/Programs/IOS/DeploymentServer/Program.cs b/Engine/Source/Programs/IOS/DeploymentServer/Program.cs
--- a/Engine/Source/Programs/IOS/DeploymentServer/Program.cs
+++ b/Engine/Source/Programs/IOS/DeploymentServer/Program.cs
@@ -9,6 +9,9 @@ using System.Runtime.Remoting.Channels.Ipc;
 using System.Diagnostics;
 using System.Runtime.InteropServices;
 using System.Collections.Generic;
+using System.IO;
+using Microsoft.Win32;
+using iPhonePackager;
 
 namespace DeploymentServer
 {
@@ -24,6 +27,36 @@ namespace DeploymentServer
 			}
 			else
 			{
+				List<string> PathBits = new List<string>();
+				PathBits.Add(Environment.GetEnvironmentVariable("Path"));
+
+				// Try to add the paths from the registry (they aren't always available on newer iTunes installs though)
+				object RegistrySupportDir = Registry.GetValue(@"HKEY_LOCAL_MACHINE\SOFTWARE\Apple Inc.\Apple Application Support", "InstallDir", Environment.CurrentDirectory);
+				if (RegistrySupportDir != null)
+				{
+					DirectoryInfo ApplicationSupportDirectory = new DirectoryInfo(RegistrySupportDir.ToString());
+
+					if (ApplicationSupportDirectory.Exists)
+					{
+						PathBits.Add(ApplicationSupportDirectory.FullName);
+					}
+				}
+
+				// Add some guesses as well
+				DirectoryInfo AppleMobileDeviceSupport = new DirectoryInfo(Environment.GetFolderPath(Environment.SpecialFolder.CommonProgramFiles) + @"\Apple\Mobile Device Support");
+				if (AppleMobileDeviceSupport.Exists)
+				{
+					PathBits.Add(AppleMobileDeviceSupport.FullName);
+				}
+
+				DirectoryInfo AppleMobileDeviceSupportX86 = new DirectoryInfo(Environment.GetFolderPath(Environment.SpecialFolder.CommonProgramFilesX86) + @"\Apple\Mobile Device Support");
+				if ((AppleMobileDeviceSupport != AppleMobileDeviceSupportX86) && (AppleMobileDeviceSupportX86.Exists))
+				{
+					PathBits.Add(AppleMobileDeviceSupportX86.FullName);
+				}
+
+				// Set the path from all the individual bits
+				Environment.SetEnvironmentVariable("Path", string.Join(";", PathBits));
 				CoreImpl = new CoreFoundationWin32();
 			}
 		}
@@ -77,6 +110,10 @@ namespace DeploymentServer
 
 	internal class CoreFoundationWin32 : CoreFoundationImpl
 	{
+		static CoreFoundationWin32()
+		{
+		}
+
 		public int RunLoopInMode(IntPtr mode, double seconds, int returnAfterSourceHandled)
 		{
 			return CFRunLoopRunInMode(mode, seconds, returnAfterSourceHandled);
@@ -119,19 +156,20 @@ namespace DeploymentServer
 					// There is no parent-child relationship WRT windows, it's self-imposed.
 					int ParentPID = int.Parse(args[1]);
 
+					DeploymentProxy.Deployer = new DeploymentImplementation();
 					IpcServerChannel Channel = new IpcServerChannel("iPhonePackager");
 					ChannelServices.RegisterChannel(Channel, false);
-					RemotingConfiguration.RegisterWellKnownServiceType(typeof(DeploymentImplementation), "DeploymentServer_PID" + ParentPID.ToString(), WellKnownObjectMode.Singleton);
+					RemotingConfiguration.RegisterWellKnownServiceType(typeof(DeploymentProxy), "DeploymentServer_PID" + ParentPID.ToString(), WellKnownObjectMode.Singleton);
 
 					Process ParentProcess = Process.GetProcessById(ParentPID);
 					while (!ParentProcess.HasExited)
 					{
-						System.Threading.Thread.Sleep(1000);
+						CoreFoundation.RunLoopRunInMode(CoreFoundation.kCFRunLoopDefaultMode(), 1.0, 0);
 					}
 				}
-				catch (System.Exception)
+				catch (System.Exception Ex)
 				{
-					
+					Console.WriteLine(Ex.Message);
 				}
             }
             else
