fix for returning false in the deployment server when only targeting one device for the operation
UE-9093, UE-9101
#ios

diff --git a/Engine/Source/Programs/IOS/DeploymentServer/DeployTime.cs b/Engine/Source/Programs/IOS/DeploymentServer/DeployTime.cs
--- a/Engine/Source/Programs/IOS/DeploymentServer/DeployTime.cs
+++ b/Engine/Source/Programs/IOS/DeploymentServer/DeployTime.cs
@@ -383,7 +383,7 @@ namespace DeploymentServer
 					return bResult;
 				}
 
-				return false;
+				return true;
             });
         }
 
@@ -395,7 +395,7 @@ namespace DeploymentServer
 			// Transfer to all connected devices
 			return PerformActionOnAllDevices(StandardEnumerationDelayMS, delegate(MobileDeviceInstance Device)
 			{
-				bool bResult = false;
+				bool bResult = true;
 				if (String.IsNullOrEmpty(DeviceId) || Device.DeviceId == DeviceId)
 				{
 					bResult = Device.TryCopy(BundleIdentifier, Manifest);
@@ -409,7 +409,7 @@ namespace DeploymentServer
         {
             return PerformActionOnAllDevices(StandardEnumerationDelayMS, delegate(MobileDeviceInstance Device)
             {
-				bool bResult = false;
+				bool bResult = true;
 				if (String.IsNullOrEmpty(DeviceId) || Device.DeviceId == DeviceId)
 				{
 					string SafeDeviceName = MobileDeviceInstance.SanitizePathNoFilename(Device.DeviceName);
@@ -432,7 +432,7 @@ namespace DeploymentServer
 		{
 			return PerformActionOnAllDevices(StandardEnumerationDelayMS, delegate(MobileDeviceInstance Device)
 			{
-				bool bResult = false;
+				bool bResult = true;
 				if (String.IsNullOrEmpty(DeviceId) || Device.DeviceId == DeviceId)
 				{
 					bResult = Device.TryBackup(BundleIdentifier, DestionationFiles);
