fix for provisions not showing "no matchin provision" if the bundle identifier didn't match no provision was found
fix for checking local times against UTC times
fix to show only the most recent certificate if two identical certificates were found
fix on Mac for certificates not turning green when they would be the ones to be selected for code signing
UE-9095, UE-9207, UE-9277, UE-9279
#ios

diff --git a/Engine/Source/Developer/iOS/IOSPlatformEditor/Private/IOSTargetSettingsCustomization.cpp b/Engine/Source/Developer/iOS/IOSPlatformEditor/Private/IOSTargetSettingsCustomization.cpp
--- a/Engine/Source/Developer/iOS/IOSPlatformEditor/Private/IOSTargetSettingsCustomization.cpp
+++ b/Engine/Source/Developer/iOS/IOSPlatformEditor/Private/IOSTargetSettingsCustomization.cpp
@@ -144,11 +144,34 @@ void FIOSTargetSettingsCustomization::UpdateStatus()
 					else if (Key.Contains(TEXT("EndDate")))
 					{
 						FString Date, Time;
-						Value.Split(TEXT(" "), &Date, &Time);
+						Value.Split(TEXT("T"), &Date, &Time);
 						Cert->Expires = Date;
 					}
 				}
-				CertificateList.Add(Cert);
+				CertificatePtr PrevCert = NULL;
+				for (int CIndex = 0; CIndex < CertificateList.Num() && !PrevCert.IsValid(); ++CIndex)
+				{
+					if (CertificateList[CIndex]->Name == Cert->Name)
+					{
+						PrevCert = CertificateList[CIndex];
+						break;
+					}
+				}
+				if (!PrevCert.IsValid())
+				{
+					CertificateList.Add(Cert);
+				}
+				else
+				{
+					FDateTime time1, time2;
+					FDateTime::ParseIso8601(*(PrevCert->Expires), time1);
+					FDateTime::ParseIso8601(*(Cert->Expires), time2);
+					if (time2 > time1)
+					{
+						PrevCert->Expires = Cert->Expires;
+					}
+					Cert = NULL;
+				}
 			}
 			else if (Line.Contains(TEXT("PROVISION-"), ESearchCase::CaseSensitive))
 			{
diff --git a/Engine/Source/Programs/IOS/iPhonePackager/CodeSigning.cs b/Engine/Source/Programs/IOS/iPhonePackager/CodeSigning.cs
--- a/Engine/Source/Programs/IOS/iPhonePackager/CodeSigning.cs
+++ b/Engine/Source/Programs/IOS/iPhonePackager/CodeSigning.cs
@@ -241,16 +241,15 @@ namespace iPhonePackager
 						string base64 = CertToolData.Substring(start, (end - start));
 						byte[] certData = Convert.FromBase64String(base64);
 						X509Certificate2 cert = new X509Certificate2(certData);
-						//@TODO: Pretty sure the certificate information from the library is in local time, not UTC and this works as expected, but it should be verified!
 						DateTime EffectiveDate = cert.NotBefore;
 						DateTime ExpirationDate = cert.NotAfter;
-						DateTime Now = DateTime.Now;
+						DateTime Now = DateTime.UtcNow;
 						string Subject = cert.Subject;
 						int SubjStart = Subject.IndexOf("CN=") + 3;
 						int SubjEnd = Subject.IndexOf(",", SubjStart);
 						cert.FriendlyName = Subject.Substring(SubjStart, (SubjEnd - SubjStart));
 						bool bCertTimeIsValid = (EffectiveDate < Now) && (ExpirationDate > Now);
-						Program.LogVerbose("CERTIFICATE-Name:{0},Validity:{1},StartDate:{2},EndDate:{3}", cert.FriendlyName, bCertTimeIsValid ? "VALID" : "EXPIRED", cert.GetEffectiveDateString(), cert.GetExpirationDateString());
+						Program.LogVerbose("CERTIFICATE-Name:{0},Validity:{1},StartDate:{2},EndDate:{3}", cert.FriendlyName, bCertTimeIsValid ? "VALID" : "EXPIRED", EffectiveDate.ToString("o"), ExpirationDate.ToString("o"));
 
 						start = CertToolData.IndexOf(header, start);
 					}
@@ -267,26 +266,24 @@ namespace iPhonePackager
 
 				foreach (X509Certificate2 TestCert in FoundCerts)
 				{
-					//@TODO: Pretty sure the certificate information from the library is in local time, not UTC and this works as expected, but it should be verified!
 					DateTime EffectiveDate = TestCert.NotBefore;
 					DateTime ExpirationDate = TestCert.NotAfter;
 					DateTime Now = DateTime.Now;
 
 					bool bCertTimeIsValid = (EffectiveDate < Now) && (ExpirationDate > Now);
-					Program.LogVerbose("CERTIFICATE-Name:{0},Validity:{1},StartDate:{2},EndDate:{3}", TestCert.FriendlyName, bCertTimeIsValid ? "VALID" : "EXPIRED", TestCert.GetEffectiveDateString(), TestCert.GetExpirationDateString());
+					Program.LogVerbose("CERTIFICATE-Name:{0},Validity:{1},StartDate:{2},EndDate:{3}", TestCert.FriendlyName, bCertTimeIsValid ? "VALID" : "EXPIRED", EffectiveDate.ToString("o"), ExpirationDate.ToString("o"));
 				}
 
 				FoundCerts = Store.Certificates.Find(X509FindType.FindBySubjectName, "iPhone Distribution", false);
 
 				foreach (X509Certificate2 TestCert in FoundCerts)
 				{
-					//@TODO: Pretty sure the certificate information from the library is in local time, not UTC and this works as expected, but it should be verified!
 					DateTime EffectiveDate = TestCert.NotBefore;
 					DateTime ExpirationDate = TestCert.NotAfter;
 					DateTime Now = DateTime.Now;
 
 					bool bCertTimeIsValid = (EffectiveDate < Now) && (ExpirationDate > Now);
-					Program.LogVerbose("CERTIFICATE-Name:{0},Validity:{1},StartDate:{2},EndDate:{3}", TestCert.FriendlyName, bCertTimeIsValid ? "VALID" : "EXPIRED", TestCert.GetEffectiveDateString(), TestCert.GetExpirationDateString());
+					Program.LogVerbose("CERTIFICATE-Name:{0},Validity:{1},StartDate:{2},EndDate:{3}", TestCert.FriendlyName, bCertTimeIsValid ? "VALID" : "EXPIRED", EffectiveDate.ToString("o"), ExpirationDate.ToString("o"));
 				}
 
 				Store.Close();
@@ -310,7 +307,7 @@ namespace iPhonePackager
 
 				DateTime EffectiveDate = p.CreationDate;
 				DateTime ExpirationDate = p.ExpirationDate;
-				DateTime Now = DateTime.Now;
+				DateTime Now = DateTime.UtcNow;
 
 				bool bCertTimeIsValid = (EffectiveDate < Now) && (ExpirationDate > Now);
 				bool bValid = false;
@@ -337,20 +334,17 @@ namespace iPhonePackager
 				}
 				bool bDistribution = ((p.ProvisionedDeviceIDs.Count == 0) && !p.bDebug);
 				string Validity = "VALID";
-				if (!bPassesNameCheck && !bPassesWildCardCheck && !bPassesCompanyCheck)
+				if (!bCertTimeIsValid)
 				{
-					Validity = "NO_MATCH";
+					Validity = "EXPIRED";
 				}
-				else
+				else if (!bValid)
 				{
-					if (!bValid)
-					{
-						Validity = "NO_CERT";
-					}
-					else if (!bCertTimeIsValid)
-					{
-						Validity = "EXPIRED";
-					}
+					Validity = "NO_CERT";
+				}
+				else if (!bPassesNameCheck && !bPassesWildCardCheck && !bPassesCompanyCheck)
+				{
+					Validity = "NO_MATCH";
 				}
 				if ((string.IsNullOrWhiteSpace(SelectedProvision) || FoundName < 2) && Validity == "VALID" && !bDistribution)
 				{
@@ -418,6 +412,10 @@ namespace iPhonePackager
 
 						if (ValidInTimeCert != null)
 						{
+							int StartIndex = SourceCert.SubjectName.Name.IndexOf("CN=") + 3;
+							int EndIndex = SourceCert.SubjectName.Name.IndexOf(", ", StartIndex);
+							SourceCert.FriendlyName = SourceCert.SubjectName.Name.Substring(StartIndex, EndIndex - StartIndex);
+
 							// Found a cert in the valid time range, quit now!
 							Result = ValidInTimeCert;
 							break;
