fix for not copying the provision from the AppData directory first before falling back to old system
UE-11130
#ios

diff --git a/Engine/Source/Programs/IOS/iPhonePackager/CompileTime.cs b/Engine/Source/Programs/IOS/iPhonePackager/CompileTime.cs
--- a/Engine/Source/Programs/IOS/iPhonePackager/CompileTime.cs
+++ b/Engine/Source/Programs/IOS/iPhonePackager/CompileTime.cs
@@ -233,16 +233,23 @@ namespace iPhonePackager
 			}
 
 			// Copy the mobile provision file over
-			string ProvisionWithPrefix = FileOperations.FindPrefixedFile(Config.BuildDirectory, Program.GameName + ".mobileprovision");
+			string CFBundleIdentifier = null;
+			Info.GetString("CFBundleIdentifier", out CFBundleIdentifier);
+			bool bNameMatch;
+			string ProvisionWithPrefix = MobileProvision.FindCompatibleProvision(CFBundleIdentifier, out bNameMatch);
 			if (!File.Exists(ProvisionWithPrefix))
 			{
-				ProvisionWithPrefix = FileOperations.FindPrefixedFile(Config.BuildDirectory + "/NotForLicensees/", Program.GameName + ".mobileprovision");
+				ProvisionWithPrefix = FileOperations.FindPrefixedFile(Config.BuildDirectory, Program.GameName + ".mobileprovision");
 				if (!File.Exists(ProvisionWithPrefix))
 				{
-					ProvisionWithPrefix = FileOperations.FindPrefixedFile(Config.EngineBuildDirectory, "UE4Game.mobileprovision");
+					ProvisionWithPrefix = FileOperations.FindPrefixedFile(Config.BuildDirectory + "/NotForLicensees/", Program.GameName + ".mobileprovision");
 					if (!File.Exists(ProvisionWithPrefix))
 					{
-						ProvisionWithPrefix = FileOperations.FindPrefixedFile(Config.EngineBuildDirectory + "/NotForLicensees/", "UE4Game.mobileprovision");
+						ProvisionWithPrefix = FileOperations.FindPrefixedFile(Config.EngineBuildDirectory, "UE4Game.mobileprovision");
+						if (!File.Exists(ProvisionWithPrefix))
+						{
+							ProvisionWithPrefix = FileOperations.FindPrefixedFile(Config.EngineBuildDirectory + "/NotForLicensees/", "UE4Game.mobileprovision");
+						}
 					}
 				}
 			}
