fix for not finding a provision when the bundle identifier has been changed
UE-11207
#ios

diff --git a/Engine/Source/Developer/iOS/IOSTargetPlatform/Private/IOSTargetPlatform.cpp b/Engine/Source/Developer/iOS/IOSTargetPlatform/Private/IOSTargetPlatform.cpp
--- a/Engine/Source/Developer/iOS/IOSTargetPlatform/Private/IOSTargetPlatform.cpp
+++ b/Engine/Source/Developer/iOS/IOSTargetPlatform/Private/IOSTargetPlatform.cpp
@@ -133,14 +133,18 @@ int32 FIOSTargetPlatform::CheckRequirements(const FString& ProjectPath, bool bPr
 #endif
 
 	// shell to IPP and get the status of the provision and cert
+
+	FString BundleIdentifier;
+	GConfig->GetString(TEXT("/Script/IOSRuntimeSettings.IOSRuntimeSettings"), TEXT("BundleIdentifier"), BundleIdentifier, GEngineIni);
+	BundleIdentifier = BundleIdentifier.Replace(TEXT("[PROJECT_NAME]"), FApp::GetGameName());
 #if PLATFORM_MAC
     FString CmdExe = TEXT("/bin/sh");
     FString ScriptPath = FPaths::ConvertRelativePathToFull(FPaths::EngineDir() / TEXT("Build/BatchFiles/Mac/RunMono.sh"));
     FString IPPPath = FPaths::ConvertRelativePathToFull(FPaths::EngineDir() / TEXT("Binaries/DotNet/IOS/IPhonePackager.exe"));
-    FString CommandLine = FString::Printf(TEXT("\"%s\" \"%s\" Validate Engine -project \"%s\""), *ScriptPath, *IPPPath, *ProjectPath);
+	FString CommandLine = FString::Printf(TEXT("\"%s\" \"%s\" Validate Engine -project \"%s\" -bundlename \"%s\""), *ScriptPath, *IPPPath, *ProjectPath, *(BundleIdentifier));
 #else
 	FString CmdExe = FPaths::ConvertRelativePathToFull(FPaths::EngineDir() / TEXT("Binaries/DotNet/IOS/IPhonePackager.exe"));
-	FString CommandLine = FString::Printf(TEXT("Validate Engine -project \"%s\""), *ProjectPath);
+	FString CommandLine = FString::Printf(TEXT("Validate Engine -project \"%s\" -bundlename \"%s\""), *ProjectPath, *(BundleIdentifier));
 #endif
 	TSharedPtr<FMonitoredProcess> IPPProcess = MakeShareable(new FMonitoredProcess(CmdExe, CommandLine, true));
 	OutputMessage = TEXT("");
diff --git a/Engine/Source/Editor/UnrealEd/Private/Kismet2/DebuggerCommands.cpp b/Engine/Source/Editor/UnrealEd/Private/Kismet2/DebuggerCommands.cpp
--- a/Engine/Source/Editor/UnrealEd/Private/Kismet2/DebuggerCommands.cpp
+++ b/Engine/Source/Editor/UnrealEd/Private/Kismet2/DebuggerCommands.cpp
@@ -1541,6 +1541,17 @@ bool FInternalPlayWorldCommandCallbacks::IsReadyToLaunchOnDevice(FString DeviceI
 			UnrecoverableError = true;
 		}
 
+		if ((Result & ETargetPlatformReadyStatus::ManifestNotFound) != 0)
+		{
+			AddMessageLog(
+				LOCTEXT("ManifestNotFound", "Manifest not found."),
+				LOCTEXT("ManifestNotFoundMessageDetail", "The generated application manifest could not be found."),
+				NotInstalledTutorialLink
+				);
+
+			UnrecoverableError = true;
+		}
+
 		if (UnrecoverableError)
 		{
 			return false;
diff --git a/Engine/Source/Programs/IOS/iPhonePackager/CodeSigning.cs b/Engine/Source/Programs/IOS/iPhonePackager/CodeSigning.cs
--- a/Engine/Source/Programs/IOS/iPhonePackager/CodeSigning.cs
+++ b/Engine/Source/Programs/IOS/iPhonePackager/CodeSigning.cs
@@ -77,47 +77,51 @@ namespace iPhonePackager
 			bHasOverridesFile = File.Exists(Config.GetPlistOverrideFilename());
 			bNameMatch = false;
 
-			// Load Info.plist, which guides nearly everything else
-			string plistFile = Config.EngineBuildDirectory + "/UE4Game-Info.plist";
-			if (!string.IsNullOrEmpty(Config.ProjectFile))
-			{
-				plistFile = Path.GetDirectoryName(Config.ProjectFile) + "/Intermediate/IOS/" + Path.GetFileNameWithoutExtension(Config.ProjectFile) + "-Info.plist";
+			string CFBundleIdentifier = Config.OverrideBundleName;
 
-				if (!File.Exists(plistFile))
+			if (string.IsNullOrEmpty(CFBundleIdentifier))
+			{
+				// Load Info.plist, which guides nearly everything else
+				string plistFile = Config.EngineBuildDirectory + "/UE4Game-Info.plist";
+				if (!string.IsNullOrEmpty(Config.ProjectFile))
 				{
-					plistFile = Config.IntermediateDirectory + "/UE4Game-Info.plist";
+					plistFile = Path.GetDirectoryName(Config.ProjectFile) + "/Intermediate/IOS/" + Path.GetFileNameWithoutExtension(Config.ProjectFile) + "-Info.plist";
+
 					if (!File.Exists(plistFile))
 					{
-						plistFile = Config.EngineBuildDirectory + "/UE4Game-Info.plist";
+						plistFile = Config.IntermediateDirectory + "/UE4Game-Info.plist";
+						if (!File.Exists(plistFile))
+						{
+							plistFile = Config.EngineBuildDirectory + "/UE4Game-Info.plist";
+						}
 					}
 				}
-			}
-			Utilities.PListHelper Info = null;
-			try
-			{
-				string RawInfoPList = File.ReadAllText(plistFile, Encoding.UTF8);
-				Info = new Utilities.PListHelper(RawInfoPList);;
-			}
-			catch (Exception ex)
-			{
-				Console.WriteLine(ex.Message);
-			}
+				Utilities.PListHelper Info = null;
+				try
+				{
+					string RawInfoPList = File.ReadAllText(plistFile, Encoding.UTF8);
+					Info = new Utilities.PListHelper(RawInfoPList); ;
+				}
+				catch (Exception ex)
+				{
+					Console.WriteLine(ex.Message);
+				}
 
-			if (Info == null)
-			{
-				return false;
-			}
+				if (Info == null)
+				{
+					return false;
+				}
 
-			// Get the name of the bundle
-			string CFBundleIdentifier = null;
-			Info.GetString("CFBundleIdentifier", out CFBundleIdentifier);
-			if (CFBundleIdentifier == null)
-			{
-				return false;
-			}
-			else
-			{
-				CFBundleIdentifier = CFBundleIdentifier.Replace("${BUNDLE_IDENTIFIER}", Path.GetFileNameWithoutExtension(Config.ProjectFile));
+				// Get the name of the bundle
+				Info.GetString("CFBundleIdentifier", out CFBundleIdentifier);
+				if (CFBundleIdentifier == null)
+				{
+					return false;
+				}
+				else
+				{
+					CFBundleIdentifier = CFBundleIdentifier.Replace("${BUNDLE_IDENTIFIER}", Path.GetFileNameWithoutExtension(Config.ProjectFile));
+				}
 			}
 
 			// Check for a mobile provision
