additional logging for file connections in Deployment Server
UE-11401
#ios

diff --git a/Engine/Source/Programs/IOS/MobileDeviceInterface/IPhoneFile.cs b/Engine/Source/Programs/IOS/MobileDeviceInterface/IPhoneFile.cs
--- a/Engine/Source/Programs/IOS/MobileDeviceInterface/IPhoneFile.cs
+++ b/Engine/Source/Programs/IOS/MobileDeviceInterface/IPhoneFile.cs
@@ -271,7 +271,7 @@ namespace Manzana
             if (ret != 0)
             {
                 phone.Reconnect();
-                throw new IOException("AFCFileRefOpen failed with error " + ret.ToString());
+                throw new IOException("AFCFileRefOpen (" + full_path + ") failed with error " + ret.ToString());
             }
 
             return new iPhoneFile(phone, handle, mode);
