if somehow both iTunes 11 and iTunes 12 are installed, prefer iTunes 12 over 11
#ios

diff --git a/Engine/Source/Programs/IOS/MobileDeviceInterface/MobileDevice.cs b/Engine/Source/Programs/IOS/MobileDeviceInterface/MobileDevice.cs
--- a/Engine/Source/Programs/IOS/MobileDeviceInterface/MobileDevice.cs
+++ b/Engine/Source/Programs/IOS/MobileDeviceInterface/MobileDevice.cs
@@ -2457,13 +2457,13 @@ namespace Manzana
 			{
 				string dllPath11 = Microsoft.Win32.Registry.GetValue("HKEY_LOCAL_MACHINE\\SOFTWARE\\Wow6432Node\\Apple Inc.\\Apple Mobile Device Support\\Shared", "iTunesMobileDeviceDLL", null) as string;
 				string dllPath12 = Microsoft.Win32.Registry.GetValue("HKEY_LOCAL_MACHINE\\SOFTWARE\\Wow6432Node\\Apple Inc.\\Apple Mobile Device Support\\Shared", "MobileDeviceDLL", null) as string;
-				if (!String.IsNullOrEmpty(dllPath11) && File.Exists(dllPath11))
+				if (!String.IsNullOrEmpty(dllPath12) && File.Exists(dllPath12))
 				{
-					DeviceImpl = new MobileDeviceWiniTunes11();
+					DeviceImpl = new MobileDeviceWiniTunes12();
 				}
-				else if (!String.IsNullOrEmpty(dllPath12) && File.Exists(dllPath12))
+				else if (!String.IsNullOrEmpty(dllPath11) && File.Exists(dllPath11))
 				{
-					DeviceImpl = new MobileDeviceWiniTunes12();
+					DeviceImpl = new MobileDeviceWiniTunes11();
 				}
 			}
 
