fix for AFCFileRefOpen error 7
UE-11302
#ios

diff --git a/Engine/Source/Programs/IOS/DeploymentServer/DeployTime.cs b/Engine/Source/Programs/IOS/DeploymentServer/DeployTime.cs
--- a/Engine/Source/Programs/IOS/DeploymentServer/DeployTime.cs
+++ b/Engine/Source/Programs/IOS/DeploymentServer/DeployTime.cs
@@ -262,7 +262,7 @@ namespace DeploymentServer
             ConnectToDevices(DelayEnumerationPeriodMS);
 
             // Keep looking at the device list and executing on new ones as long as we keep finding them
-            Dictionary<MobileDeviceInstance, bool> Runs = new Dictionary<MobileDeviceInstance, bool>();
+			Dictionary<string, bool> Runs = new Dictionary<string, bool>();
             bool bFoundNewDevices = true;
             while (bFoundNewDevices)
             {
@@ -271,16 +271,18 @@ namespace DeploymentServer
                 bFoundNewDevices = false;
                 foreach (MobileDeviceInstance PotentialDevice in Devices)
                 {
-                    if (!Runs.ContainsKey(PotentialDevice))
+					PotentialDevice.Reconnect();
+					string DeviceId = PotentialDevice.DeviceId;
+					if (DeviceId == String.Empty)
+					{
+						System.Threading.Thread.Sleep((int)SleepAfterFailedDeviceCallMS);
+						DeviceId = PotentialDevice.DeviceId;
+					}
+
+					if (!Runs.ContainsKey(DeviceId))
                     {
-                        string DeviceName = PotentialDevice.DeviceName;
-                        if (DeviceName == String.Empty)
-                        {
-                            System.Threading.Thread.Sleep((int)SleepAfterFailedDeviceCallMS);
-                        }
-
                         // New device, do the work on it
-                        Runs.Add(PotentialDevice, PerDeviceWork(PotentialDevice));
+						Runs.Add(DeviceId, PerDeviceWork(PotentialDevice));
                         bFoundNewDevices = true;
                     }
                 }
diff --git a/Engine/Source/Programs/IOS/MobileDeviceInterface/IPhoneFile.cs b/Engine/Source/Programs/IOS/MobileDeviceInterface/IPhoneFile.cs
--- a/Engine/Source/Programs/IOS/MobileDeviceInterface/IPhoneFile.cs
+++ b/Engine/Source/Programs/IOS/MobileDeviceInterface/IPhoneFile.cs
@@ -271,6 +271,9 @@ namespace Manzana
             if (ret != 0)
             {
                 phone.Reconnect();
+				// error codes and their meaning
+				// 7 - bad path
+				// 18 - out of disk space
                 throw new IOException("AFCFileRefOpen (" + full_path + ") failed with error " + ret.ToString());
             }
 
diff --git a/Engine/Source/Programs/IOS/MobileDeviceInterface/MobileDeviceInstance.cs b/Engine/Source/Programs/IOS/MobileDeviceInterface/MobileDeviceInstance.cs
--- a/Engine/Source/Programs/IOS/MobileDeviceInterface/MobileDeviceInstance.cs
+++ b/Engine/Source/Programs/IOS/MobileDeviceInterface/MobileDeviceInstance.cs
@@ -1027,7 +1027,15 @@ namespace Manzana
 
             // Make sure the directory exists on the phone
             string DirectoryOnPhone = PathOnPhone.Remove(PathOnPhone.LastIndexOf('/'));
-            CreateDirectory(DirectoryOnPhone);
+			if (!IsDirectory(DirectoryOnPhone))
+			{
+				Console.WriteLine("Directory (" + DirectoryOnPhone + ") doesn't exist!");
+				if (!CreateDirectory(DirectoryOnPhone))
+				{
+					Console.WriteLine("CreateDirectory (" + DirectoryOnPhone + ") failed");
+//					throw new IOException("CreateDirectory (" + DirectoryOnPhone + ") failed");
+				}
+			}
 
             FileStream SourceFile = File.OpenRead(PathOnPC);
             iPhoneFile DestinationFile = iPhoneFile.OpenWrite(this, PathOnPhone);
@@ -1304,7 +1312,7 @@ namespace Manzana
         public void CopyFileToPublicStaging(string SourceFile)
         {
             string IpaFilename = Path.GetFileName(SourceFile);
-            CopyFileToPhone(SourceFile, "/PublicStaging/" + IpaFilename);
+            CopyFileToPhone(SourceFile, "PublicStaging/" + IpaFilename);
             //Dictionary<string, string> fi = GetFileInfo("/PublicStaging/" + IpaFilename);
         }
 
