fix for losing device connection when trying to copy to the device
UE-11302
#ios

diff --git a/Engine/Source/Programs/IOS/MobileDeviceInterface/MobileDeviceInstance.cs b/Engine/Source/Programs/IOS/MobileDeviceInterface/MobileDeviceInstance.cs
--- a/Engine/Source/Programs/IOS/MobileDeviceInterface/MobileDeviceInstance.cs
+++ b/Engine/Source/Programs/IOS/MobileDeviceInterface/MobileDeviceInstance.cs
@@ -835,7 +835,11 @@ namespace Manzana
         {
             if (!IsConnected)
             {
-                throw new Exception("Not connected to phone");
+				Reconnect();
+				if (!IsConnected)
+				{
+					throw new Exception("Not connected to phone");
+				}
             }
 
             IntPtr hAFCDir = IntPtr.Zero;
@@ -1025,8 +1029,14 @@ namespace Manzana
 
             byte[] buffer = new byte[ChunkSize];
 
-            // Make sure the directory exists on the phone
-            string DirectoryOnPhone = PathOnPhone.Remove(PathOnPhone.LastIndexOf('/'));
+			// verify we are still connected
+			if (!IsConnected)
+			{
+				Reconnect();
+			}
+
+			// Make sure the directory exists on the phone
+			string DirectoryOnPhone = PathOnPhone.Remove(PathOnPhone.LastIndexOf('/'));
 			if (!IsDirectory(DirectoryOnPhone))
 			{
 				Console.WriteLine("Directory (" + DirectoryOnPhone + ") doesn't exist!");
@@ -1037,10 +1047,9 @@ namespace Manzana
 				}
 			}
 
+					
             FileStream SourceFile = File.OpenRead(PathOnPC);
             iPhoneFile DestinationFile = iPhoneFile.OpenWrite(this, PathOnPhone);
-
-            
             long ProgressInterval = Math.Max(buffer.Length, (SourceFile.Length / TransferProgressDivisor));
 
             long NextProgressPrintout = ProgressInterval;
@@ -1324,27 +1333,32 @@ namespace Manzana
             // Make sure we can connect to the phone and that it has been paired with this machine previously
 			if (MobileDevice.DeviceImpl.Connect(iPhoneHandle) != 0)
             {
+				Console.WriteLine("Connect: Failed to Connect");
                 return false;
             }
 
 			if (MobileDevice.DeviceImpl.IsPaired(iPhoneHandle) != 1)
             {
+				Console.WriteLine("Connect: Device is not paired");
                 return false;
             }
 
 			if (MobileDevice.DeviceImpl.ValidatePairing(iPhoneHandle) != 0)
             {
+				Console.WriteLine("Connect: Pairing not valid");
                 return false;
             }
 
             // Start a session
 			if (MobileDevice.DeviceImpl.StartSession(iPhoneHandle) == 1)
             {
+				Console.WriteLine("Connect: Couldn't start session");
                 return false;
             }
 
 			if (MobileDevice.DeviceImpl.StartService(iPhoneHandle, MobileDevice.CFStringMakeConstantString("com.apple.afc"), ref hService, IntPtr.Zero) != 0)
             {
+				Console.WriteLine("Connect: Couldn't Start AFC service");
                 return false;
             }
 
@@ -1352,6 +1366,7 @@ namespace Manzana
             // Open a file sharing connection
 			if (MobileDevice.DeviceImpl.ConnectionOpen(hService, 0, out AFCCommsHandle) != 0)
             {
+				Console.WriteLine("Connect: Couldn't Open File Sharing Connection");
                 return false;
             }
 
