PR #1010 - fix for DeploymentServer.exe exception
#ios

diff --git a/Engine/Source/Programs/IOS/MobileDeviceInterface/MobileDevice.cs b/Engine/Source/Programs/IOS/MobileDeviceInterface/MobileDevice.cs
--- a/Engine/Source/Programs/IOS/MobileDeviceInterface/MobileDevice.cs
+++ b/Engine/Source/Programs/IOS/MobileDeviceInterface/MobileDevice.cs
@@ -915,7 +915,6 @@ namespace Manzana
 		static MobileDeviceWiniTunes11()
 		{
 			List<string> PathBits = new List<string>();
-			PathBits.Add(Environment.GetEnvironmentVariable("Path"));
 
 			// Try to add the paths from the registry (they aren't always available on newer iTunes installs though)
 			object RegistryDllPath = Registry.GetValue(@"HKEY_LOCAL_MACHINE\SOFTWARE\Apple Inc.\Apple Mobile Device Support\Shared", "iTunesMobileDeviceDLL", DLLName);
@@ -953,6 +952,9 @@ namespace Manzana
 				PathBits.Add(AppleMobileDeviceSupportX86.FullName);
 			}
 
+			// add the rest of the path
+			PathBits.Add(Environment.GetEnvironmentVariable("Path"));
+
 			// Set the path from all the individual bits
 			Environment.SetEnvironmentVariable("Path", string.Join(";", PathBits));
 		}
@@ -1683,7 +1685,6 @@ namespace Manzana
 		static MobileDeviceWiniTunes12()
 		{
 			List<string> PathBits = new List<string>();
-			PathBits.Add(Environment.GetEnvironmentVariable("Path"));
 
 			// Try to add the paths from the registry (they aren't always available on newer iTunes installs though)
 			object RegistryDllPath = Registry.GetValue(@"HKEY_LOCAL_MACHINE\SOFTWARE\Apple Inc.\Apple Mobile Device Support\Shared", "MobileDeviceDLL", DLLName);
@@ -1721,6 +1722,9 @@ namespace Manzana
 				PathBits.Add(AppleMobileDeviceSupportX86.FullName);
 			}
 
+			// add the rest of the path
+			PathBits.Add(Environment.GetEnvironmentVariable("Path"));
+
 			// Set the path from all the individual bits
 			Environment.SetEnvironmentVariable("Path", string.Join(";", PathBits));
 		}
