fix for invalid quotes when using SSH
UE-14493
#ios

diff --git a/Engine/Source/Programs/IOS/iPhonePackager/CompileTime.cs b/Engine/Source/Programs/IOS/iPhonePackager/CompileTime.cs
--- a/Engine/Source/Programs/IOS/iPhonePackager/CompileTime.cs
+++ b/Engine/Source/Programs/IOS/iPhonePackager/CompileTime.cs
@@ -81,9 +81,18 @@ namespace iPhonePackager
 			{
 				CmdLine += String.Format(" CODE_SIGN_RESOURCE_RULES_PATH='{0}/CustomResourceRules.plist'", MacStagingRootDir);
 			}
-			CmdLine += String.Format(" CODE_SIGN_IDENTITY=\\\"{0}\\\"", Config.CodeSigningIdentity);
+			if (Config.bUseRPCUtil)
+			{
+				CmdLine += String.Format(" CODE_SIGN_IDENTITY=\\\"{0}\\\"", Config.CodeSigningIdentity);
+
+				CmdLine += String.Format(" IPHONEOS_DEPLOYMENT_TARGET=\\\"{0}\\\"", Config.MinOSVersion);
+			}
+			else
+			{
+				CmdLine += String.Format(" CODE_SIGN_IDENTITY=\"{0}\"", Config.CodeSigningIdentity);
 
-			CmdLine += String.Format(" IPHONEOS_DEPLOYMENT_TARGET=\\\"{0}\\\"", Config.MinOSVersion);
+				CmdLine += String.Format(" IPHONEOS_DEPLOYMENT_TARGET=\"{0}\"", Config.MinOSVersion);
+			}
 
 			return CmdLine;
 		}
