fixed issue with read-only mobile provisions in the Provisioning Profiles directory
UE-11318
#ios

diff --git a/Engine/Source/Programs/IOS/iPhonePackager/MobileProvisionUtilities.cs b/Engine/Source/Programs/IOS/iPhonePackager/MobileProvisionUtilities.cs
--- a/Engine/Source/Programs/IOS/iPhonePackager/MobileProvisionUtilities.cs
+++ b/Engine/Source/Programs/IOS/iPhonePackager/MobileProvisionUtilities.cs
@@ -62,9 +62,11 @@ namespace iPhonePackager
 					{
 						if (!File.Exists(Config.ProvisionDirectory + Path.GetFileName(Provision)) || File.GetLastWriteTime(Config.ProvisionDirectory + Path.GetFileName(Provision)) < File.GetLastWriteTime(Provision))
 						{
-							File.Copy(Provision, Config.ProvisionDirectory + Path.GetFileName(Provision), true);
 							FileInfo DestFileInfo = new FileInfo(Config.ProvisionDirectory + Path.GetFileName(Provision));
 							DestFileInfo.Attributes = DestFileInfo.Attributes & ~FileAttributes.ReadOnly;
+							File.Copy(Provision, Config.ProvisionDirectory + Path.GetFileName(Provision), true);
+							DestFileInfo = new FileInfo(Config.ProvisionDirectory + Path.GetFileName(Provision));
+							DestFileInfo.Attributes = DestFileInfo.Attributes & ~FileAttributes.ReadOnly;
 						}
 					}
 				}
@@ -78,9 +80,11 @@ namespace iPhonePackager
 					{
 						if (!File.Exists(Config.ProvisionDirectory + Path.GetFileName(Provision)) || File.GetLastWriteTime(Config.ProvisionDirectory + Path.GetFileName(Provision)) < File.GetLastWriteTime(Provision))
 						{
-							File.Copy(Provision, Config.ProvisionDirectory + Path.GetFileName(Provision), true);
 							FileInfo DestFileInfo = new FileInfo(Config.ProvisionDirectory + Path.GetFileName(Provision));
 							DestFileInfo.Attributes = DestFileInfo.Attributes & ~FileAttributes.ReadOnly;
+							File.Copy(Provision, Config.ProvisionDirectory + Path.GetFileName(Provision), true);
+							DestFileInfo = new FileInfo(Config.ProvisionDirectory + Path.GetFileName(Provision));
+							DestFileInfo.Attributes = DestFileInfo.Attributes & ~FileAttributes.ReadOnly;
 						}
 					}
 				}
