fix for trying to get file info for a file that doesn't exist yet
UE-15289
#ios

diff --git a/Engine/Source/Programs/IOS/iPhonePackager/MobileProvisionUtilities.cs b/Engine/Source/Programs/IOS/iPhonePackager/MobileProvisionUtilities.cs
--- a/Engine/Source/Programs/IOS/iPhonePackager/MobileProvisionUtilities.cs
+++ b/Engine/Source/Programs/IOS/iPhonePackager/MobileProvisionUtilities.cs
@@ -62,8 +62,12 @@ namespace iPhonePackager
 					{
 						if (!File.Exists(Config.ProvisionDirectory + Path.GetFileName(Provision)) || File.GetLastWriteTime(Config.ProvisionDirectory + Path.GetFileName(Provision)) < File.GetLastWriteTime(Provision))
 						{
-							FileInfo DestFileInfo = new FileInfo(Config.ProvisionDirectory + Path.GetFileName(Provision));
-							DestFileInfo.Attributes = DestFileInfo.Attributes & ~FileAttributes.ReadOnly;
+							FileInfo DestFileInfo;
+							if (File.Exists(Config.ProvisionDirectory + Path.GetFileName(Provision)))
+							{
+								DestFileInfo = new FileInfo(Config.ProvisionDirectory + Path.GetFileName(Provision));
+								DestFileInfo.Attributes = DestFileInfo.Attributes & ~FileAttributes.ReadOnly;
+							}
 							File.Copy(Provision, Config.ProvisionDirectory + Path.GetFileName(Provision), true);
 							DestFileInfo = new FileInfo(Config.ProvisionDirectory + Path.GetFileName(Provision));
 							DestFileInfo.Attributes = DestFileInfo.Attributes & ~FileAttributes.ReadOnly;
@@ -80,8 +84,12 @@ namespace iPhonePackager
 					{
 						if (!File.Exists(Config.ProvisionDirectory + Path.GetFileName(Provision)) || File.GetLastWriteTime(Config.ProvisionDirectory + Path.GetFileName(Provision)) < File.GetLastWriteTime(Provision))
 						{
-							FileInfo DestFileInfo = new FileInfo(Config.ProvisionDirectory + Path.GetFileName(Provision));
-							DestFileInfo.Attributes = DestFileInfo.Attributes & ~FileAttributes.ReadOnly;
+							FileInfo DestFileInfo;
+							if (File.Exists(Config.ProvisionDirectory + Path.GetFileName(Provision)))
+							{
+								DestFileInfo = new FileInfo(Config.ProvisionDirectory + Path.GetFileName(Provision));
+								DestFileInfo.Attributes = DestFileInfo.Attributes & ~FileAttributes.ReadOnly;
+							}
 							File.Copy(Provision, Config.ProvisionDirectory + Path.GetFileName(Provision), true);
 							DestFileInfo = new FileInfo(Config.ProvisionDirectory + Path.GetFileName(Provision));
 							DestFileInfo.Attributes = DestFileInfo.Attributes & ~FileAttributes.ReadOnly;
