addition of message to alert the user when they have not properly installed their certificates on the remote mac
UE-15490
#ios

diff --git a/Engine/Source/Editor/UnrealEd/Private/Analytics/EditorAnalytics.cpp b/Engine/Source/Editor/UnrealEd/Private/Analytics/EditorAnalytics.cpp
--- a/Engine/Source/Editor/UnrealEd/Private/Analytics/EditorAnalytics.cpp
+++ b/Engine/Source/Editor/UnrealEd/Private/Analytics/EditorAnalytics.cpp
@@ -174,6 +174,8 @@ FString FEditorAnalytics::TranslateErrorCode(int32 ErrorCode)
 		return TEXT("Failed to launch on device.  Make sure your device has been enabled for development from within the Xcode Devices window.");
 	case EAnalyticsErrorCodes::DeviceOSNewerThanSDK:
 		return TEXT("Failed to launch on device.  Make sure your install of Xcode matches or is newer than the OS on your device.");
+	case EAnalyticsErrorCodes::RemoteCertificatesNotFound:
+		return TEXT("Failed to sign executable.  Make sure your developer certificates have been installed in the System Keychain on the remote Mac.");
 	}
 	return TEXT("Unknown Error");
 }
diff --git a/Engine/Source/Editor/UnrealEd/Public/EditorAnalytics.h b/Engine/Source/Editor/UnrealEd/Public/EditorAnalytics.h
--- a/Engine/Source/Editor/UnrealEd/Public/EditorAnalytics.h
+++ b/Engine/Source/Editor/UnrealEd/Public/EditorAnalytics.h
@@ -54,6 +54,7 @@ namespace EAnalyticsErrorCodes
 		AndroidBuildToolsPathNotFound = 54,
 		NoApkSuitableForArchitecture = 55,
 		FilesInstallFailed = 56,
+		RemoteCertificatesNotFound = 57,
 		LauncherFailed = 100,
 		UATLaunchFailure = 101,
 		FailedToDeleteStagingDirectory = 102,
diff --git a/Engine/Source/Programs/AutomationTool/Program.cs b/Engine/Source/Programs/AutomationTool/Program.cs
--- a/Engine/Source/Programs/AutomationTool/Program.cs
+++ b/Engine/Source/Programs/AutomationTool/Program.cs
@@ -55,6 +55,7 @@ namespace AutomationTool
         Error_AndroidBuildToolsPathNotFound = 54,
         Error_NoApkSuitableForArchitecture = 55,
 		Error_FilesInstallFailed = 56,
+		Error_RemoteCertificatesNotFound = 57,
 		Error_LauncherFailed = 100,
 		Error_UATLaunchFailure = 101,
 		Error_FailedToDeleteStagingDirectory = 102,
diff --git a/Engine/Source/Programs/IOS/iPhonePackager/CompileTime.cs b/Engine/Source/Programs/IOS/iPhonePackager/CompileTime.cs
--- a/Engine/Source/Programs/IOS/iPhonePackager/CompileTime.cs
+++ b/Engine/Source/Programs/IOS/iPhonePackager/CompileTime.cs
@@ -308,6 +308,7 @@ namespace iPhonePackager
 			string TempKeychain = "$HOME/Library/Keychains/UE4TempKeychain.keychain";
 			string Certificate = "XcodeSupportFiles/" + MacSigningIdentityFilename;
 			string LoginKeychain = "$HOME/Library/Keychains/login.keychain";
+			ErrorCodes Error = ErrorCodes.Error_Unknown;
 
 			switch (RPCCommand.ToLowerInvariant())
 			{
@@ -365,6 +366,7 @@ namespace iPhonePackager
 				DisplayCommandLine = CurrentBaseXCodeCommandLine;
 				CommandLine = "\"" + MacXcodeStagingDir + "/..\" " + DisplayCommandLine;
 				WorkingFolder = "\"" + MacXcodeStagingDir + "/..\"";
+				Error = ErrorCodes.Error_RemoteCertificatesNotFound;
 				break;
 
 			case "createkeychain":
@@ -494,6 +496,11 @@ namespace iPhonePackager
 			{
 				Program.Log("Running SSH on " + MacName + " ... ");
 				bSuccess = SSHCommandHelper.Command(MacName, DisplayCommandLine, WorkingFolder);
+				if (bSuccess == false)
+				{
+					Program.Error("RPCCommand {0} failed with return code {1}", RPCCommand, Error);
+					Program.ReturnCode = (int)Error;
+				}
 			}
 
 
diff --git a/Engine/Source/Programs/IOS/iPhonePackager/iPhonePackager.cs b/Engine/Source/Programs/IOS/iPhonePackager/iPhonePackager.cs
--- a/Engine/Source/Programs/IOS/iPhonePackager/iPhonePackager.cs
+++ b/Engine/Source/Programs/IOS/iPhonePackager/iPhonePackager.cs
@@ -62,6 +62,7 @@ namespace iPhonePackager
         Error_AndroidBuildToolsPathNotFound = 54,
         Error_NoApkSuitableForArchitecture = 55,
 		Error_FilesInstallFailed = 56,
+		Error_RemoteCertificatesNotFound = 57,
 		Error_LauncherFailed = 100,
         Error_UATLaunchFailure = 101,
         Error_FailedToDeleteStagingDirectory = 102,
diff --git a/Engine/Source/Programs/UnrealBuildTool/IOS/IOSToolChain.cs b/Engine/Source/Programs/UnrealBuildTool/IOS/IOSToolChain.cs
--- a/Engine/Source/Programs/UnrealBuildTool/IOS/IOSToolChain.cs
+++ b/Engine/Source/Programs/UnrealBuildTool/IOS/IOSToolChain.cs
@@ -1270,9 +1270,10 @@ namespace UnrealBuildTool
 
 					OutputReceivedDataEventHandlerEncounteredError = false;
 					OutputReceivedDataEventHandlerEncounteredErrorMessage = "";
-					Utils.RunLocalProcess(StubGenerateProcess);
+					int ExitCode = Utils.RunLocalProcess(StubGenerateProcess);
 					if (OutputReceivedDataEventHandlerEncounteredError)
 					{
+						UnrealBuildTool.ExtendedErrorCode = ExitCode;
 						throw new Exception(OutputReceivedDataEventHandlerEncounteredErrorMessage);
 					}
 					
diff --git a/Engine/Source/Programs/UnrealBuildTool/System/UnrealBuildTool.cs b/Engine/Source/Programs/UnrealBuildTool/System/UnrealBuildTool.cs
--- a/Engine/Source/Programs/UnrealBuildTool/System/UnrealBuildTool.cs
+++ b/Engine/Source/Programs/UnrealBuildTool/System/UnrealBuildTool.cs
@@ -1357,9 +1357,14 @@ namespace UnrealBuildTool
             // Print some performance info
             Log.TraceVerbose("Execution time: {0}", (DateTime.UtcNow - StartTime - MutexWaitTime).TotalSeconds);
 
+			if (ExtendedErrorCode != 0)
+			{
+				return ExtendedErrorCode;
+			}
             return (int)Result; 
         }
 
+		public static int ExtendedErrorCode = 0;
         private static int Main(string[] Arguments)
         {
             // make sure we catch any exceptions and return an appropriate error code.
@@ -1383,6 +1388,10 @@ namespace UnrealBuildTool
                 {
                     Log.TraceError("UnrealBuildTool Exception: " + Exception.ToString());
                 }
+				if (ExtendedErrorCode != 0)
+				{
+					return ExtendedErrorCode;
+				}
                 return (int)ECompilationResult.OtherCompilationError;
             }
         }
diff --git a/Engine/Source/Programs/UnrealFrontend/Private/Windows/WindowsUnrealFrontendMain.cpp b/Engine/Source/Programs/UnrealFrontend/Private/Windows/WindowsUnrealFrontendMain.cpp
--- a/Engine/Source/Programs/UnrealFrontend/Private/Windows/WindowsUnrealFrontendMain.cpp
+++ b/Engine/Source/Programs/UnrealFrontend/Private/Windows/WindowsUnrealFrontendMain.cpp
@@ -19,6 +19,7 @@ int32 WINAPI WinMain( HINSTANCE hInInstance, HINSTANCE hPrevInstance, char* lpCm
 	hInstance = hInInstance;
 
 	const TCHAR* CmdLine = FPlatformMisc::GetCompleteCommandLine();
+	CmdLine = FCommandLine::RemoveExeName(CmdLine);
 
 #if !UE_BUILD_SHIPPING
 	if (FParse::Param(CmdLine, TEXT("crashreports")))
diff --git a/Engine/Source/Runtime/Core/Private/Windows/WindowsPlatformMisc.cpp b/Engine/Source/Runtime/Core/Private/Windows/WindowsPlatformMisc.cpp
--- a/Engine/Source/Runtime/Core/Private/Windows/WindowsPlatformMisc.cpp
+++ b/Engine/Source/Runtime/Core/Private/Windows/WindowsPlatformMisc.cpp
@@ -2436,8 +2436,7 @@ const TCHAR* FWindowsPlatformMisc::GetCompleteCommandLine()
 	static FString CommandLine;
 	if (CommandLine.Len() == 0)
 	{
-		const TCHAR* OrgCmdLine = ::GetCommandLine();
-		CommandLine = FCommandLine::RemoveExeName(OrgCmdLine);
+		CommandLine = ::GetCommandLine();
 
 		// Get the path to the arguments file
 		FString ArgsFileName = FPaths::RootDir() / TEXT("UE4CommandLine.txt");
