fix for date time check for certificates on Mac
#ios

diff --git a/Engine/Source/Programs/IOS/iPhonePackager/CodeSigning.cs b/Engine/Source/Programs/IOS/iPhonePackager/CodeSigning.cs
--- a/Engine/Source/Programs/IOS/iPhonePackager/CodeSigning.cs
+++ b/Engine/Source/Programs/IOS/iPhonePackager/CodeSigning.cs
@@ -318,6 +318,7 @@ namespace iPhonePackager
 				X509Certificate2 Cert = FindCertificate(p);
 				if (Cert != null)
 				{
+					Now = (Environment.OSVersion.Platform == PlatformID.Unix || Environment.OSVersion.Platform == PlatformID.MacOSX) ? Now : DateTime.Now;
 					bValid = (Cert.NotBefore < Now) && (Cert.NotAfter > Now);
 				}
 				bool bPassesNameCheck = p.ApplicationIdentifier.Substring(p.ApplicationIdentifierPrefix.Length+1) == CFBundleIdentifier;
