UE-18895 - can now deploy unicode paths/filenames
#ios

diff --git a/Engine/Source/Programs/IOS/MobileDeviceInterface/CoreFoundation.cs b/Engine/Source/Programs/IOS/MobileDeviceInterface/CoreFoundation.cs
--- a/Engine/Source/Programs/IOS/MobileDeviceInterface/CoreFoundation.cs
+++ b/Engine/Source/Programs/IOS/MobileDeviceInterface/CoreFoundation.cs
@@ -76,6 +76,7 @@ namespace Manzana
 		TypedPtr<CFString> __CFStringMakeConstantString(byte[] s);
 		TypedPtr<CFURL> CFURLCreateWithFileSystemPath(IntPtr Allocator, TypedPtr<CFString> FilePath, CFURLPathStyle PathStyle, int isDirectory);
 		Boolean CFStringGetCString(TypedPtr<CFString> theString, byte[] buffer, int bufferSize, CFStringBuiltInEncodings encoding);
+		TypedPtr<CFString> CFStringCreateWithCString(IntPtr allocator, byte[] buffer);
 		TypedPtr<CFString> CFURLGetString(IntPtr anURL);
 		uint CFGetTypeID(IntPtr FromInstance);
 		uint CFStringGetTypeID();
@@ -99,6 +100,7 @@ namespace Manzana
 			/* const CFDictionaryKeyCallBacks* */ IntPtr keyCallBacks,
 			/* const CFDictionaryValueCallBacks* */ IntPtr valueCallBacks
 		);
+		Boolean CFStringGetFileSystemRepresentation(TypedPtr<CFString> theString, byte[] buffer, int bufferSize);
 	}
 
 	internal class CoreFoundationWin : CoreFoundationImpl
@@ -109,6 +111,11 @@ namespace Manzana
 			return CoreFoundation.__CFStringMakeConstantString(s);
 		}
 
+		public TypedPtr<CFString> CFStringCreateWithCString(IntPtr allocator, byte[] buffer)
+		{
+			return CoreFoundation.CFStringCreateWithCString(allocator, buffer, CFStringBuiltInEncodings.kCFStringEncodingUTF8);
+		}
+
 		public TypedPtr<CFURL> CFURLCreateWithFileSystemPath(IntPtr Allocator, TypedPtr<CFString> FilePath, CFURLPathStyle PathStyle, int isDirectory)
 		{
 			return CoreFoundation.CFURLCreateWithFileSystemPath(Allocator, (IntPtr)FilePath, PathStyle, isDirectory);
@@ -214,11 +221,19 @@ namespace Manzana
 			return CoreFoundation.CFDictionaryCreateMutable(allocator, capacity, keyCallBacks, valueCallBacks);
 		}
 
+		public Boolean CFStringGetFileSystemRepresentation(TypedPtr<CFString> theString, byte[] buffer, int bufferSize)
+		{
+			return CoreFoundation.CFStringGetFileSystemRepresentation((IntPtr)theString, buffer, bufferSize);
+		}
+
 		private class CoreFoundation
 		{
 			[DllImport("CoreFoundation.dll", CallingConvention = CallingConvention.Cdecl)]
 			public static extern IntPtr __CFStringMakeConstantString(byte[] s);
 
+			[DllImport("CoreFoundation.dll", CallingConvention = CallingConvention.Cdecl)]
+			public static extern IntPtr CFStringCreateWithCString(IntPtr allocator, byte[] buffer, CFStringBuiltInEncodings encoding);
+			
 			[DllImport("CoreFoundation.dll", CallingConvention = CallingConvention.Cdecl)]
 			public static extern IntPtr/*CFURL*/ CFURLCreateWithFileSystemPath(IntPtr Allocator, IntPtr/*CFString*/ FilePath, CFURLPathStyle PathStyle, int isDirectory);
 
@@ -286,6 +301,9 @@ namespace Manzana
 
 			[DllImport("CoreFoundation.dll", CallingConvention = CallingConvention.Cdecl)]
 			public static extern void CFPreferencesSetAppValue(IntPtr Key, IntPtr Value, IntPtr ApplicationID);
+
+			[DllImport("CoreFoundation.dll", CallingConvention = CallingConvention.Cdecl)]
+			public static extern Boolean CFStringGetFileSystemRepresentation(IntPtr theString, byte[] buffer, int bufferSize);
 		}
 	}
 
@@ -297,6 +315,11 @@ namespace Manzana
 			return CoreFoundation.__CFStringMakeConstantString(s);
 		}
 
+		public TypedPtr<CFString> CFStringCreateWithCString(IntPtr allocator, byte[] buffer)
+		{
+			return CoreFoundation.CFStringCreateWithCString(allocator, buffer, CFStringBuiltInEncodings.kCFStringEncodingUTF8);
+		}
+
 		public TypedPtr<CFURL> CFURLCreateWithFileSystemPath(IntPtr Allocator, TypedPtr<CFString> FilePath, CFURLPathStyle PathStyle, int isDirectory)
 		{
 			return CoreFoundation.CFURLCreateWithFileSystemPath(Allocator, (IntPtr)FilePath, PathStyle, isDirectory);
@@ -402,11 +425,19 @@ namespace Manzana
 			return CoreFoundation.CFDictionaryCreateMutable(allocator, capacity, keyCallBacks, valueCallBacks);
 		}
 
+		public Boolean CFStringGetFileSystemRepresentation(TypedPtr<CFString> theString, byte[] buffer, int bufferSize)
+		{
+			return CoreFoundation.CFStringGetFileSystemRepresentation((IntPtr)theString, buffer, bufferSize);
+		}
+
 		private class CoreFoundation
 		{
 			[DllImport("/System/Library/Frameworks/CoreFoundation.framework/CoreFoundation", CallingConvention = CallingConvention.Cdecl)]
 			public static extern IntPtr __CFStringMakeConstantString(byte[] s);
 
+			[DllImport("/System/Library/Frameworks/CoreFoundation.framework/CoreFoundation", CallingConvention = CallingConvention.Cdecl)]
+			public static extern IntPtr CFStringCreateWithCString(IntPtr allocator, byte[] buffer, CFStringBuiltInEncodings encoding);
+
 			[DllImport("/System/Library/Frameworks/CoreFoundation.framework/CoreFoundation", CallingConvention = CallingConvention.Cdecl)]
 			public static extern IntPtr/*CFURL*/ CFURLCreateWithFileSystemPath(IntPtr Allocator, IntPtr/*CFString*/ FilePath, CFURLPathStyle PathStyle, int isDirectory);
 
@@ -474,6 +505,9 @@ namespace Manzana
 
 			[DllImport("/System/Library/Frameworks/CoreFoundation.framework/CoreFoundation", CallingConvention = CallingConvention.Cdecl)]
 			public static extern void CFPreferencesSetAppValue(IntPtr Key, IntPtr Value, IntPtr ApplicationID);
+
+			[DllImport("/System/Library/Frameworks/CoreFoundation.framework/CoreFoundation", CallingConvention = CallingConvention.Cdecl)]
+			public static extern Boolean CFStringGetFileSystemRepresentation(IntPtr theString, byte[] buffer, int bufferSize);
 		}
 	}
 
@@ -550,7 +584,34 @@ namespace Manzana
             return Encoding.UTF8.GetString(bytes, 0, ValidLength);
         }
 
-        public static string GetStringForUrl(TypedPtr<CFURL> Url)
+		public static string StringToFileSystemRepresentation(string InString)
+		{
+			TypedPtr<CFString> cfString = CoreImpl.CFStringCreateWithCString(kCFAllocatorDefault, StringToCString(InString));
+			return CFStringGetFileSystemRepresentation(cfString);
+		}
+
+		public static string CFStringGetFileSystemRepresentation(TypedPtr<CFString> InString)
+		{
+			byte[] bytes = new byte[2048];
+			CoreImpl.CFStringGetFileSystemRepresentation(InString.Handle, bytes, 2048);
+
+			int ValidLength = 0;
+			foreach (byte b in bytes)
+			{
+				if (b == 0)
+				{
+					break;
+				}
+				else
+				{
+					ValidLength++;
+				}
+			}
+
+			return Encoding.UTF8.GetString(bytes, 0, ValidLength);
+		}
+
+		public static string GetStringForUrl(TypedPtr<CFURL> Url)
         {
 			TypedPtr<CFString> cfString = CoreImpl.CFURLGetString((IntPtr)Url);
             return CFStringGetCString((IntPtr)cfString);
diff --git a/Engine/Source/Programs/IOS/MobileDeviceInterface/IPhoneFile.cs b/Engine/Source/Programs/IOS/MobileDeviceInterface/IPhoneFile.cs
--- a/Engine/Source/Programs/IOS/MobileDeviceInterface/IPhoneFile.cs
+++ b/Engine/Source/Programs/IOS/MobileDeviceInterface/IPhoneFile.cs
@@ -267,7 +267,7 @@ namespace Manzana
             }
 
             full_path = phone.FullPath(phone.CurrentDirectory, path);
-			ret = MobileDevice.DeviceImpl.FileRefOpen(phone.AFCHandle, full_path, (Int64)mode, out handle);
+			ret = MobileDevice.DeviceImpl.FileRefOpen(phone.AFCHandle, MobileDevice.StringToFileSystemRepresentation(full_path), (Int64)mode, out handle);
             if (ret != 0)
             {
                 phone.Reconnect();
diff --git a/Engine/Source/Programs/IOS/MobileDeviceInterface/MobileDeviceInstance.cs b/Engine/Source/Programs/IOS/MobileDeviceInterface/MobileDeviceInstance.cs
--- a/Engine/Source/Programs/IOS/MobileDeviceInterface/MobileDeviceInstance.cs
+++ b/Engine/Source/Programs/IOS/MobileDeviceInterface/MobileDeviceInstance.cs
@@ -693,7 +693,7 @@ namespace Manzana
             string full_path = FullPath(CurrentDirectory, path);
 
             IntPtr hAFCDir = IntPtr.Zero;
-            if (MobileDevice.DeviceImpl.DirectoryOpen(AFCCommsHandle, full_path, ref hAFCDir) != 0)
+            if (MobileDevice.DeviceImpl.DirectoryOpen(AFCCommsHandle, MobileDevice.StringToFileSystemRepresentation(full_path), ref hAFCDir) != 0)
             {
                 throw new Exception("Path does not exist");
             }
@@ -730,7 +730,7 @@ namespace Manzana
             Dictionary<string, string> ans = new Dictionary<string, string>();
             TypedPtr<AFCDictionary> Data;
 
-			int ret = MobileDevice.DeviceImpl.FileInfoOpen(AFCCommsHandle, path, out Data);
+			int ret = MobileDevice.DeviceImpl.FileInfoOpen(AFCCommsHandle, MobileDevice.StringToFileSystemRepresentation(path), out Data);
             if ((ret == 0) && (Data.Handle != IntPtr.Zero))
             {
                 IntPtr pname;
@@ -795,7 +795,7 @@ namespace Manzana
                 // test for symbolic directory link
                 IntPtr hAFCDir = IntPtr.Zero;
 
-				if (directory = (MobileDevice.DeviceImpl.DirectoryOpen(AFCCommsHandle, path, ref hAFCDir) == 0))
+				if (directory = (MobileDevice.DeviceImpl.DirectoryOpen(AFCCommsHandle, MobileDevice.StringToFileSystemRepresentation(path), ref hAFCDir) == 0))
                 {
 					MobileDevice.DeviceImpl.DirectoryClose(AFCCommsHandle, hAFCDir);
                 }
@@ -823,7 +823,7 @@ namespace Manzana
         /// <returns>true if directory was created</returns>
         public bool CreateDirectory(string path)
         {
-			return !(MobileDevice.DeviceImpl.DirectoryCreate(AFCCommsHandle, FullPath(CurrentDirectory, path)) != 0);
+			return !(MobileDevice.DeviceImpl.DirectoryCreate(AFCCommsHandle, MobileDevice.StringToFileSystemRepresentation(FullPath(CurrentDirectory, path))) != 0);
         }
 
         /// <summary>
@@ -846,7 +846,7 @@ namespace Manzana
             string full_path = FullPath(CurrentDirectory, path);
             //full_path = "/private"; // bug test
 
-			int res = MobileDevice.DeviceImpl.DirectoryOpen(AFCCommsHandle, full_path, ref hAFCDir);
+			int res = MobileDevice.DeviceImpl.DirectoryOpen(AFCCommsHandle, MobileDevice.StringToFileSystemRepresentation(full_path), ref hAFCDir);
             if (res != 0)
             {
                 throw new Exception("Path does not exist: " + res.ToString());
@@ -876,7 +876,9 @@ namespace Manzana
         ///	<remarks>Files cannot be moved across filesystem boundaries.</remarks>
         public bool Rename(string sourceName, string destName)
         {
-			return MobileDevice.DeviceImpl.RenamePath(AFCCommsHandle, FullPath(CurrentDirectory, sourceName), FullPath(CurrentDirectory, destName)) == 0;
+			return MobileDevice.DeviceImpl.RenamePath(AFCCommsHandle,
+				MobileDevice.StringToFileSystemRepresentation(FullPath(CurrentDirectory, sourceName)),
+				MobileDevice.StringToFileSystemRepresentation(FullPath(CurrentDirectory, destName))) == 0;
         }
 
         /// <summary>
@@ -898,7 +900,7 @@ namespace Manzana
         {
             TypedPtr<AFCDictionary> data = IntPtr.Zero;
 
-			int ret = MobileDevice.DeviceImpl.FileInfoOpen(AFCCommsHandle, path, out data);
+			int ret = MobileDevice.DeviceImpl.FileInfoOpen(AFCCommsHandle, MobileDevice.StringToFileSystemRepresentation(path), out data);
             if (ret == 0)
             {
 				MobileDevice.DeviceImpl.KeyValueClose(data);
@@ -950,7 +952,7 @@ namespace Manzana
             string full_path = FullPath(CurrentDirectory, path);
             if (IsDirectory(full_path))
             {
-				MobileDevice.DeviceImpl.RemovePath(AFCCommsHandle, full_path);
+				MobileDevice.DeviceImpl.RemovePath(AFCCommsHandle, MobileDevice.StringToFileSystemRepresentation(full_path));
             }
         }
 
@@ -984,7 +986,7 @@ namespace Manzana
             string full_path = FullPath(CurrentDirectory, path);
             if (Exists(full_path))
             {
-				MobileDevice.DeviceImpl.RemovePath(AFCCommsHandle, full_path);
+				MobileDevice.DeviceImpl.RemovePath(AFCCommsHandle, MobileDevice.StringToFileSystemRepresentation(full_path));
             }
         }
         #endregion	// Filesystem
