UE-18895 - fix for failure to deploy to iOS when there are non-Ascii characters in the filename or path
#ios

diff --git a/Engine/Source/Programs/IOS/MobileDeviceInterface/CoreFoundation.cs b/Engine/Source/Programs/IOS/MobileDeviceInterface/CoreFoundation.cs
--- a/Engine/Source/Programs/IOS/MobileDeviceInterface/CoreFoundation.cs
+++ b/Engine/Source/Programs/IOS/MobileDeviceInterface/CoreFoundation.cs
@@ -28,7 +28,13 @@ namespace Manzana
 		kCFStringEncodingASCII = 0x00000600,
 		kCFStringEncodingUnicode = 0x00000100,
 		kCFStringEncodingUTF8 = 0x08000100,
-		kCFStringEncodingNonLossyASCII = 0x00000BFF
+		kCFStringEncodingNonLossyASCII = 0x00000BFF,
+		kCFStringEncodingUTF16 = 0x00000100,
+		kCFStringEncodingUTF16BE = 0x10000100,
+		kCFStringEncodingUTF16LE = 0x14000100,
+		kCFStringEncodingUTF32 = 0x0c000100,
+		kCFStringEncodingUTF32BE = 0x18000100,
+		kCFStringEncodingUTF32LE = 0x1c000100
 	};
 
 	public enum CFURLPathStyle
@@ -76,6 +82,7 @@ namespace Manzana
 		TypedPtr<CFString> __CFStringMakeConstantString(byte[] s);
 		TypedPtr<CFURL> CFURLCreateWithFileSystemPath(IntPtr Allocator, TypedPtr<CFString> FilePath, CFURLPathStyle PathStyle, int isDirectory);
 		Boolean CFStringGetCString(TypedPtr<CFString> theString, byte[] buffer, int bufferSize, CFStringBuiltInEncodings encoding);
+		TypedPtr<CFString> CFStringCreateWithBytes(IntPtr allocator, byte[] buffer, int bufferSize, CFStringBuiltInEncodings encoding, Boolean isExternalRep);
 		TypedPtr<CFString> CFStringCreateWithCString(IntPtr allocator, byte[] buffer);
 		TypedPtr<CFString> CFURLGetString(IntPtr anURL);
 		uint CFGetTypeID(IntPtr FromInstance);
@@ -116,6 +123,11 @@ namespace Manzana
 			return CoreFoundation.CFStringCreateWithCString(allocator, buffer, CFStringBuiltInEncodings.kCFStringEncodingUTF8);
 		}
 
+		public TypedPtr<CFString> CFStringCreateWithBytes(IntPtr allocator, byte[] buffer, int bufferSize, CFStringBuiltInEncodings encoding, Boolean isExternalRep)
+		{
+			return CoreFoundation.CFStringCreateWithBytes(allocator, buffer, bufferSize, encoding, isExternalRep);
+		}
+
 		public TypedPtr<CFURL> CFURLCreateWithFileSystemPath(IntPtr Allocator, TypedPtr<CFString> FilePath, CFURLPathStyle PathStyle, int isDirectory)
 		{
 			return CoreFoundation.CFURLCreateWithFileSystemPath(Allocator, (IntPtr)FilePath, PathStyle, isDirectory);
@@ -233,6 +245,9 @@ namespace Manzana
 
 			[DllImport("CoreFoundation.dll", CallingConvention = CallingConvention.Cdecl)]
 			public static extern IntPtr CFStringCreateWithCString(IntPtr allocator, byte[] buffer, CFStringBuiltInEncodings encoding);
+
+			[DllImport("CoreFoundation.dll", CallingConvention = CallingConvention.Cdecl)]
+			public static extern IntPtr CFStringCreateWithBytes(IntPtr allocator, byte[] buffer, int bufferSize, CFStringBuiltInEncodings encoding, Boolean isExtRep);
 			
 			[DllImport("CoreFoundation.dll", CallingConvention = CallingConvention.Cdecl)]
 			public static extern IntPtr/*CFURL*/ CFURLCreateWithFileSystemPath(IntPtr Allocator, IntPtr/*CFString*/ FilePath, CFURLPathStyle PathStyle, int isDirectory);
@@ -320,6 +335,11 @@ namespace Manzana
 			return CoreFoundation.CFStringCreateWithCString(allocator, buffer, CFStringBuiltInEncodings.kCFStringEncodingUTF8);
 		}
 
+		public TypedPtr<CFString> CFStringCreateWithBytes(IntPtr allocator, byte[] buffer, int bufferSize, CFStringBuiltInEncodings encoding, Boolean isExternalRep)
+		{
+			return CoreFoundation.CFStringCreateWithBytes(allocator, buffer, bufferSize, encoding, isExternalRep);
+		}
+
 		public TypedPtr<CFURL> CFURLCreateWithFileSystemPath(IntPtr Allocator, TypedPtr<CFString> FilePath, CFURLPathStyle PathStyle, int isDirectory)
 		{
 			return CoreFoundation.CFURLCreateWithFileSystemPath(Allocator, (IntPtr)FilePath, PathStyle, isDirectory);
@@ -438,6 +458,9 @@ namespace Manzana
 			[DllImport("/System/Library/Frameworks/CoreFoundation.framework/CoreFoundation", CallingConvention = CallingConvention.Cdecl)]
 			public static extern IntPtr CFStringCreateWithCString(IntPtr allocator, byte[] buffer, CFStringBuiltInEncodings encoding);
 
+			[DllImport("/System/Library/Frameworks/CoreFoundation.framework/CoreFoundation", CallingConvention = CallingConvention.Cdecl)]
+			public static extern IntPtr CFStringCreateWithBytes(IntPtr allocator, byte[] buffer, int bufferSize, CFStringBuiltInEncodings encoding, Boolean isExtRep);
+
 			[DllImport("/System/Library/Frameworks/CoreFoundation.framework/CoreFoundation", CallingConvention = CallingConvention.Cdecl)]
 			public static extern IntPtr/*CFURL*/ CFURLCreateWithFileSystemPath(IntPtr Allocator, IntPtr/*CFString*/ FilePath, CFURLPathStyle PathStyle, int isDirectory);
 
@@ -553,6 +576,16 @@ namespace Manzana
             return bytes;
         }
 
+		private static byte[] StringToUTF16String(string value)
+		{
+			return Encoding.Unicode.GetBytes(value);
+		}
+
+		private static byte[] StringToUTF32LEString(string value)
+		{
+			return Encoding.UTF32.GetBytes(value);
+		}
+
         public static TypedPtr<CFString> CFStringMakeConstantString(string s)
         {
             return CoreImpl.__CFStringMakeConstantString(StringToCString(s));
@@ -584,13 +617,14 @@ namespace Manzana
             return Encoding.UTF8.GetString(bytes, 0, ValidLength);
         }
 
-		public static string StringToFileSystemRepresentation(string InString)
+		public static byte[] StringToFileSystemRepresentation(string InString)
 		{
-			TypedPtr<CFString> cfString = CoreImpl.CFStringCreateWithCString(kCFAllocatorDefault, StringToCString(InString));
+			byte[] BString =  StringToUTF32LEString(InString);
+			TypedPtr<CFString> cfString = CoreImpl.CFStringCreateWithBytes(kCFAllocatorDefault, BString, BString.Length, CFStringBuiltInEncodings.kCFStringEncodingUTF32LE, false);
 			return CFStringGetFileSystemRepresentation(cfString);
 		}
 
-		public static string CFStringGetFileSystemRepresentation(TypedPtr<CFString> InString)
+		public static byte[] CFStringGetFileSystemRepresentation(TypedPtr<CFString> InString)
 		{
 			byte[] bytes = new byte[2048];
 			CoreImpl.CFStringGetFileSystemRepresentation(InString.Handle, bytes, 2048);
@@ -608,7 +642,8 @@ namespace Manzana
 				}
 			}
 
-			return Encoding.UTF8.GetString(bytes, 0, ValidLength);
+//			return Encoding.UTF8.GetString(bytes, 0, ValidLength);
+			return bytes;
 		}
 
 		public static string GetStringForUrl(TypedPtr<CFURL> Url)
diff --git a/Engine/Source/Programs/IOS/MobileDeviceInterface/IPhoneFile.cs b/Engine/Source/Programs/IOS/MobileDeviceInterface/IPhoneFile.cs
--- a/Engine/Source/Programs/IOS/MobileDeviceInterface/IPhoneFile.cs
+++ b/Engine/Source/Programs/IOS/MobileDeviceInterface/IPhoneFile.cs
@@ -267,7 +267,7 @@ namespace Manzana
             }
 
             full_path = phone.FullPath(phone.CurrentDirectory, path);
-			ret = MobileDevice.DeviceImpl.FileRefOpen(phone.AFCHandle, MobileDevice.StringToFileSystemRepresentation(full_path), (Int64)mode, out handle);
+			ret = MobileDevice.DeviceImpl.FileRefOpen(phone.AFCHandle, full_path, (Int64)mode, out handle);
             if (ret != 0)
             {
                 phone.Reconnect();
diff --git a/Engine/Source/Programs/IOS/MobileDeviceInterface/MobileDevice.cs b/Engine/Source/Programs/IOS/MobileDeviceInterface/MobileDevice.cs
--- a/Engine/Source/Programs/IOS/MobileDeviceInterface/MobileDevice.cs
+++ b/Engine/Source/Programs/IOS/MobileDeviceInterface/MobileDevice.cs
@@ -165,7 +165,6 @@ namespace Manzana
 		int DeviceInfoOpen(IntPtr handle, ref IntPtr dict);
 		int DirectoryClose(TypedPtr<AFCCommConnection> conn, IntPtr dir);
 		int DirectoryCreate(TypedPtr<AFCCommConnection> conn, string path);
-		int DirectoryOpen(TypedPtr<AFCCommConnection> conn, byte[] path, ref IntPtr dir);
 		int DirectoryRead(TypedPtr<AFCCommConnection> conn, IntPtr dir, ref IntPtr dirent);
 		int FileInfoOpen(TypedPtr<AFCCommConnection> conn, string path, out TypedPtr<AFCDictionary> OutDict);
 		int FileRefClose(TypedPtr<AFCCommConnection> conn, Int64 handle);
@@ -344,7 +343,7 @@ namespace Manzana
 			return AFC.DirectoryCreate(conn, path);
 		}
 
-		public int DirectoryOpen(TypedPtr<AFCCommConnection> conn, byte[] path, ref IntPtr dir)
+		public int DirectoryOpen(TypedPtr<AFCCommConnection> conn, string path, ref IntPtr dir)
 		{
 			return AFC.DirectoryOpen(conn, path, ref dir);
 		}
@@ -694,11 +693,6 @@ namespace Manzana
 
 #region AFC Operations
 
-		public int DirectoryOpen(TypedPtr<AFCCommConnection> conn, string path, ref IntPtr dir)
-		{
-			return AFC.DirectoryOpen((IntPtr)conn, Encoding.UTF8.GetBytes(path), ref dir);
-		}
-
 		public int DirectoryRead(TypedPtr<AFCCommConnection> conn, IntPtr dir, ref string buffer)
 		{
 			int ret;
@@ -753,12 +747,12 @@ namespace Manzana
 
 			public static int DirectoryCreate(TypedPtr<AFCCommConnection> conn, string path)
 			{
-				return AFCDirectoryCreate((IntPtr)conn, path);
+				return AFCDirectoryCreate((IntPtr)conn, MobileDevice.StringToFileSystemRepresentation(path));
 			}
 
-			public static int DirectoryOpen(TypedPtr<AFCCommConnection> conn, byte[] path, ref IntPtr dir)
+			public static int DirectoryOpen(TypedPtr<AFCCommConnection> conn, string path, ref IntPtr dir)
 			{
-				return AFCDirectoryOpen((IntPtr)conn, path, ref dir);
+				return AFCDirectoryOpen((IntPtr)conn, MobileDevice.StringToFileSystemRepresentation(path), ref dir);
 			}
 
 			public static int DirectoryRead(TypedPtr<AFCCommConnection> conn, IntPtr dir, ref IntPtr dirent)
@@ -769,7 +763,7 @@ namespace Manzana
 			public static int FileInfoOpen(TypedPtr<AFCCommConnection> conn, string path, out TypedPtr<AFCDictionary> OutDict)
 			{
 				IntPtr UntypedDict;
-				int Result = AFCFileInfoOpen((IntPtr)conn, path, out UntypedDict);
+				int Result = AFCFileInfoOpen((IntPtr)conn, MobileDevice.StringToFileSystemRepresentation(path), out UntypedDict);
 				OutDict = UntypedDict;
 
 				return Result;
@@ -782,7 +776,7 @@ namespace Manzana
 
 			public static int FileRefOpen(TypedPtr<AFCCommConnection> conn, string path, Int64 mode, out Int64 handle)
 			{
-				return AFCFileRefOpen((IntPtr)conn, path, mode, out handle);
+				return AFCFileRefOpen((IntPtr)conn, MobileDevice.StringToFileSystemRepresentation(path), mode, out handle);
 			}
 
 			public static int FileRefRead(TypedPtr<AFCCommConnection> conn, Int64 handle, byte[] buffer, ref uint len)
@@ -827,12 +821,12 @@ namespace Manzana
 
 			public static int RemovePath(TypedPtr<AFCCommConnection> conn, string path)
 			{
-				return AFCRemovePath((IntPtr)conn, path);
+				return AFCRemovePath((IntPtr)conn, MobileDevice.StringToFileSystemRepresentation(path));
 			}
 
 			public static int RenamePath(TypedPtr<AFCCommConnection> conn, string OldPath, string NewPath)
 			{
-				return AFCRenamePath((IntPtr)conn, OldPath, NewPath);
+				return AFCRenamePath((IntPtr)conn, MobileDevice.StringToFileSystemRepresentation(OldPath), MobileDevice.StringToFileSystemRepresentation(NewPath));
 			}
 
 			[DllImport(DLLName, CallingConvention = CallingConvention.Cdecl)]
@@ -854,7 +848,7 @@ namespace Manzana
 			private extern static int AFCDirectoryClose(IntPtr/*AFCCommConnection*/ conn, IntPtr dir);
 
 			[DllImport(DLLName, CallingConvention = CallingConvention.Cdecl)]
-			private extern static int AFCDirectoryCreate(IntPtr/*AFCCommConnection*/ conn, string path);
+			private extern static int AFCDirectoryCreate(IntPtr/*AFCCommConnection*/ conn, byte[] path);
 
 			[DllImport(DLLName, CallingConvention = CallingConvention.Cdecl)]
 			private extern static int AFCDirectoryOpen(IntPtr/*AFCCommConnection*/ conn, byte[] path, ref IntPtr dir);
@@ -863,13 +857,13 @@ namespace Manzana
 			private extern static int AFCDirectoryRead(IntPtr/*AFCCommConnection*/ conn, IntPtr dir, ref IntPtr dirent);
 
 			[DllImport(DLLName, CallingConvention = CallingConvention.Cdecl)]
-			private extern static int AFCFileInfoOpen(IntPtr/*AFCCommConnection*/ conn, string path, out IntPtr OutDict);
+			private extern static int AFCFileInfoOpen(IntPtr/*AFCCommConnection*/ conn, byte[] path, out IntPtr OutDict);
 
 			[DllImport(DLLName, CallingConvention = CallingConvention.Cdecl)]
 			private extern static int AFCFileRefClose(IntPtr/*AFCCommConnection*/ conn, Int64 handle);
 
 			[DllImport(DLLName, CallingConvention = CallingConvention.Cdecl)]
-			private extern static int AFCFileRefOpen(IntPtr/*AFCCommConnection*/ conn, string path, Int64 mode, out Int64 handle);
+			private extern static int AFCFileRefOpen(IntPtr/*AFCCommConnection*/ conn, byte[] path, Int64 mode, out Int64 handle);
 
 			[DllImport(DLLName, CallingConvention = CallingConvention.Cdecl)]
 			private extern static int AFCFileRefRead(IntPtr/*AFCCommConnection*/ conn, Int64 handle, byte[] buffer, ref uint len);
@@ -898,10 +892,10 @@ namespace Manzana
 
 
 			[DllImport(DLLName, CallingConvention = CallingConvention.Cdecl)]
-			private extern static int AFCRemovePath(IntPtr/*AFCCommConnection*/ conn, string path);
+			private extern static int AFCRemovePath(IntPtr/*AFCCommConnection*/ conn, byte[] path);
 
 			[DllImport(DLLName, CallingConvention = CallingConvention.Cdecl)]
-			private extern static int AFCRenamePath(IntPtr/*AFCCommConnection*/ conn, string old_path, string new_path);
+			private extern static int AFCRenamePath(IntPtr/*AFCCommConnection*/ conn, byte[] old_path, byte[] new_path);
 		}
 
 		#endregion
@@ -1114,7 +1108,7 @@ namespace Manzana
 			return AFC.DirectoryCreate(conn, path);
 		}
 
-		public int DirectoryOpen(TypedPtr<AFCCommConnection> conn, byte[] path, ref IntPtr dir)
+		public int DirectoryOpen(TypedPtr<AFCCommConnection> conn, string path, ref IntPtr dir)
 		{
 			return AFC.DirectoryOpen(conn, path, ref dir);
 		}
@@ -1464,11 +1458,6 @@ namespace Manzana
 
 		#region AFC Operations
 
-		public int DirectoryOpen(TypedPtr<AFCCommConnection> conn, string path, ref IntPtr dir)
-		{
-			return AFC.DirectoryOpen((IntPtr)conn, Encoding.UTF8.GetBytes(path), ref dir);
-		}
-
 		public int DirectoryRead(TypedPtr<AFCCommConnection> conn, IntPtr dir, ref string buffer)
 		{
 			int ret;
@@ -1523,12 +1512,12 @@ namespace Manzana
 
 			public static int DirectoryCreate(TypedPtr<AFCCommConnection> conn, string path)
 			{
-				return AFCDirectoryCreate((IntPtr)conn, path);
+				return AFCDirectoryCreate((IntPtr)conn, MobileDevice.StringToFileSystemRepresentation(path));
 			}
 
-			public static int DirectoryOpen(TypedPtr<AFCCommConnection> conn, byte[] path, ref IntPtr dir)
+			public static int DirectoryOpen(TypedPtr<AFCCommConnection> conn, string path, ref IntPtr dir)
 			{
-				return AFCDirectoryOpen((IntPtr)conn, path, ref dir);
+				return AFCDirectoryOpen((IntPtr)conn, MobileDevice.StringToFileSystemRepresentation(path), ref dir);
 			}
 
 			public static int DirectoryRead(TypedPtr<AFCCommConnection> conn, IntPtr dir, ref IntPtr dirent)
@@ -1539,7 +1528,7 @@ namespace Manzana
 			public static int FileInfoOpen(TypedPtr<AFCCommConnection> conn, string path, out TypedPtr<AFCDictionary> OutDict)
 			{
 				IntPtr UntypedDict;
-				int Result = AFCFileInfoOpen((IntPtr)conn, path, out UntypedDict);
+				int Result = AFCFileInfoOpen((IntPtr)conn, MobileDevice.StringToFileSystemRepresentation(path), out UntypedDict);
 				OutDict = UntypedDict;
 
 				return Result;
@@ -1552,7 +1541,7 @@ namespace Manzana
 
 			public static int FileRefOpen(TypedPtr<AFCCommConnection> conn, string path, Int64 mode, out Int64 handle)
 			{
-				return AFCFileRefOpen((IntPtr)conn, path, mode, out handle);
+				return AFCFileRefOpen((IntPtr)conn, MobileDevice.StringToFileSystemRepresentation(path), mode, out handle);
 			}
 
 			public static int FileRefRead(TypedPtr<AFCCommConnection> conn, Int64 handle, byte[] buffer, ref uint len)
@@ -1597,12 +1586,12 @@ namespace Manzana
 
 			public static int RemovePath(TypedPtr<AFCCommConnection> conn, string path)
 			{
-				return AFCRemovePath((IntPtr)conn, path);
+				return AFCRemovePath((IntPtr)conn, MobileDevice.StringToFileSystemRepresentation(path));
 			}
 
 			public static int RenamePath(TypedPtr<AFCCommConnection> conn, string OldPath, string NewPath)
 			{
-				return AFCRenamePath((IntPtr)conn, OldPath, NewPath);
+				return AFCRenamePath((IntPtr)conn, MobileDevice.StringToFileSystemRepresentation(OldPath), MobileDevice.StringToFileSystemRepresentation(NewPath));
 			}
 
 			[DllImport(DLLName, CallingConvention = CallingConvention.Cdecl)]
@@ -1624,7 +1613,7 @@ namespace Manzana
 			private extern static int AFCDirectoryClose(IntPtr/*AFCCommConnection*/ conn, IntPtr dir);
 
 			[DllImport(DLLName, CallingConvention = CallingConvention.Cdecl)]
-			private extern static int AFCDirectoryCreate(IntPtr/*AFCCommConnection*/ conn, string path);
+			private extern static int AFCDirectoryCreate(IntPtr/*AFCCommConnection*/ conn, byte[] path);
 
 			[DllImport(DLLName, CallingConvention = CallingConvention.Cdecl)]
 			private extern static int AFCDirectoryOpen(IntPtr/*AFCCommConnection*/ conn, byte[] path, ref IntPtr dir);
@@ -1633,13 +1622,13 @@ namespace Manzana
 			private extern static int AFCDirectoryRead(IntPtr/*AFCCommConnection*/ conn, IntPtr dir, ref IntPtr dirent);
 
 			[DllImport(DLLName, CallingConvention = CallingConvention.Cdecl)]
-			private extern static int AFCFileInfoOpen(IntPtr/*AFCCommConnection*/ conn, string path, out IntPtr OutDict);
+			private extern static int AFCFileInfoOpen(IntPtr/*AFCCommConnection*/ conn, byte[] path, out IntPtr OutDict);
 
 			[DllImport(DLLName, CallingConvention = CallingConvention.Cdecl)]
 			private extern static int AFCFileRefClose(IntPtr/*AFCCommConnection*/ conn, Int64 handle);
 
 			[DllImport(DLLName, CallingConvention = CallingConvention.Cdecl)]
-			private extern static int AFCFileRefOpen(IntPtr/*AFCCommConnection*/ conn, string path, Int64 mode, out Int64 handle);
+			private extern static int AFCFileRefOpen(IntPtr/*AFCCommConnection*/ conn, byte[] path, Int64 mode, out Int64 handle);
 
 			[DllImport(DLLName, CallingConvention = CallingConvention.Cdecl)]
 			private extern static int AFCFileRefRead(IntPtr/*AFCCommConnection*/ conn, Int64 handle, byte[] buffer, ref uint len);
@@ -1668,10 +1657,10 @@ namespace Manzana
 
 
 			[DllImport(DLLName, CallingConvention = CallingConvention.Cdecl)]
-			private extern static int AFCRemovePath(IntPtr/*AFCCommConnection*/ conn, string path);
+			private extern static int AFCRemovePath(IntPtr/*AFCCommConnection*/ conn, byte[] path);
 
 			[DllImport(DLLName, CallingConvention = CallingConvention.Cdecl)]
-			private extern static int AFCRenamePath(IntPtr/*AFCCommConnection*/ conn, string old_path, string new_path);
+			private extern static int AFCRenamePath(IntPtr/*AFCCommConnection*/ conn, byte[] old_path, byte[] new_path);
 		}
 
 		#endregion
@@ -1884,7 +1873,7 @@ namespace Manzana
 			return AFC.DirectoryCreate(conn, path);
 		}
 
-		public int DirectoryOpen(TypedPtr<AFCCommConnection> conn, byte[] path, ref IntPtr dir)
+		public int DirectoryOpen(TypedPtr<AFCCommConnection> conn, string path, ref IntPtr dir)
 		{
 			return AFC.DirectoryOpen(conn, path, ref dir);
 		}
@@ -2234,11 +2223,6 @@ namespace Manzana
 
 		#region AFC Operations
 
-		public int DirectoryOpen(TypedPtr<AFCCommConnection> conn, string path, ref IntPtr dir)
-		{
-			return AFC.DirectoryOpen((IntPtr)conn, Encoding.UTF8.GetBytes(path), ref dir);
-		}
-
 		public int DirectoryRead(TypedPtr<AFCCommConnection> conn, IntPtr dir, ref string buffer)
 		{
 			int ret;
@@ -2293,12 +2277,12 @@ namespace Manzana
 
 			public static int DirectoryCreate(TypedPtr<AFCCommConnection> conn, string path)
 			{
-				return AFCDirectoryCreate((IntPtr)conn, path);
+				return AFCDirectoryCreate((IntPtr)conn, MobileDevice.StringToFileSystemRepresentation(path));
 			}
 
-			public static int DirectoryOpen(TypedPtr<AFCCommConnection> conn, byte[] path, ref IntPtr dir)
+			public static int DirectoryOpen(TypedPtr<AFCCommConnection> conn, string path, ref IntPtr dir)
 			{
-				return AFCDirectoryOpen((IntPtr)conn, path, ref dir);
+				return AFCDirectoryOpen((IntPtr)conn, MobileDevice.StringToFileSystemRepresentation(path), ref dir);
 			}
 
 			public static int DirectoryRead(TypedPtr<AFCCommConnection> conn, IntPtr dir, ref IntPtr dirent)
@@ -2309,7 +2293,7 @@ namespace Manzana
 			public static int FileInfoOpen(TypedPtr<AFCCommConnection> conn, string path, out TypedPtr<AFCDictionary> OutDict)
 			{
 				IntPtr UntypedDict;
-				int Result = AFCFileInfoOpen((IntPtr)conn, path, out UntypedDict);
+				int Result = AFCFileInfoOpen((IntPtr)conn, MobileDevice.StringToFileSystemRepresentation(path), out UntypedDict);
 				OutDict = UntypedDict;
 
 				return Result;
@@ -2322,7 +2306,7 @@ namespace Manzana
 
 			public static int FileRefOpen(TypedPtr<AFCCommConnection> conn, string path, Int64 mode, out Int64 handle)
 			{
-				return AFCFileRefOpen((IntPtr)conn, path, mode, out handle);
+				return AFCFileRefOpen((IntPtr)conn, MobileDevice.StringToFileSystemRepresentation(path), mode, out handle);
 			}
 
 			public static int FileRefRead(TypedPtr<AFCCommConnection> conn, Int64 handle, byte[] buffer, ref uint len)
@@ -2367,12 +2351,12 @@ namespace Manzana
 
 			public static int RemovePath(TypedPtr<AFCCommConnection> conn, string path)
 			{
-				return AFCRemovePath((IntPtr)conn, path);
+				return AFCRemovePath((IntPtr)conn, MobileDevice.StringToFileSystemRepresentation(path));
 			}
 
 			public static int RenamePath(TypedPtr<AFCCommConnection> conn, string OldPath, string NewPath)
 			{
-				return AFCRenamePath((IntPtr)conn, OldPath, NewPath);
+				return AFCRenamePath((IntPtr)conn, MobileDevice.StringToFileSystemRepresentation(OldPath), MobileDevice.StringToFileSystemRepresentation(NewPath));
 			}
 
 			[DllImport(DLLName, CallingConvention = CallingConvention.Cdecl)]
@@ -2394,7 +2378,7 @@ namespace Manzana
 			private extern static int AFCDirectoryClose(IntPtr/*AFCCommConnection*/ conn, IntPtr dir);
 
 			[DllImport(DLLName, CallingConvention = CallingConvention.Cdecl)]
-			private extern static int AFCDirectoryCreate(IntPtr/*AFCCommConnection*/ conn, string path);
+			private extern static int AFCDirectoryCreate(IntPtr/*AFCCommConnection*/ conn, byte[] path);
 
 			[DllImport(DLLName, CallingConvention = CallingConvention.Cdecl)]
 			private extern static int AFCDirectoryOpen(IntPtr/*AFCCommConnection*/ conn, byte[] path, ref IntPtr dir);
@@ -2403,13 +2387,13 @@ namespace Manzana
 			private extern static int AFCDirectoryRead(IntPtr/*AFCCommConnection*/ conn, IntPtr dir, ref IntPtr dirent);
 
 			[DllImport(DLLName, CallingConvention = CallingConvention.Cdecl)]
-			private extern static int AFCFileInfoOpen(IntPtr/*AFCCommConnection*/ conn, string path, out IntPtr OutDict);
+			private extern static int AFCFileInfoOpen(IntPtr/*AFCCommConnection*/ conn, byte[] path, out IntPtr OutDict);
 
 			[DllImport(DLLName, CallingConvention = CallingConvention.Cdecl)]
 			private extern static int AFCFileRefClose(IntPtr/*AFCCommConnection*/ conn, Int64 handle);
 
 			[DllImport(DLLName, CallingConvention = CallingConvention.Cdecl)]
-			private extern static int AFCFileRefOpen(IntPtr/*AFCCommConnection*/ conn, string path, Int64 mode, out Int64 handle);
+			private extern static int AFCFileRefOpen(IntPtr/*AFCCommConnection*/ conn, byte[] path, Int64 mode, out Int64 handle);
 
 			[DllImport(DLLName, CallingConvention = CallingConvention.Cdecl)]
 			private extern static int AFCFileRefRead(IntPtr/*AFCCommConnection*/ conn, Int64 handle, byte[] buffer, ref uint len);
@@ -2438,10 +2422,10 @@ namespace Manzana
 
 
 			[DllImport(DLLName, CallingConvention = CallingConvention.Cdecl)]
-			private extern static int AFCRemovePath(IntPtr/*AFCCommConnection*/ conn, string path);
+			private extern static int AFCRemovePath(IntPtr/*AFCCommConnection*/ conn, byte[] path);
 
 			[DllImport(DLLName, CallingConvention = CallingConvention.Cdecl)]
-			private extern static int AFCRenamePath(IntPtr/*AFCCommConnection*/ conn, string old_path, string new_path);
+			private extern static int AFCRenamePath(IntPtr/*AFCCommConnection*/ conn, byte[] old_path, byte[] new_path);
 		}
 
 		#endregion
diff --git a/Engine/Source/Programs/IOS/MobileDeviceInterface/MobileDeviceInstance.cs b/Engine/Source/Programs/IOS/MobileDeviceInterface/MobileDeviceInstance.cs
--- a/Engine/Source/Programs/IOS/MobileDeviceInterface/MobileDeviceInstance.cs
+++ b/Engine/Source/Programs/IOS/MobileDeviceInterface/MobileDeviceInstance.cs
@@ -693,7 +693,7 @@ namespace Manzana
             string full_path = FullPath(CurrentDirectory, path);
 
             IntPtr hAFCDir = IntPtr.Zero;
-            if (MobileDevice.DeviceImpl.DirectoryOpen(AFCCommsHandle, MobileDevice.StringToFileSystemRepresentation(full_path), ref hAFCDir) != 0)
+            if (MobileDevice.DeviceImpl.DirectoryOpen(AFCCommsHandle, full_path, ref hAFCDir) != 0)
             {
                 throw new Exception("Path does not exist");
             }
@@ -730,7 +730,7 @@ namespace Manzana
             Dictionary<string, string> ans = new Dictionary<string, string>();
             TypedPtr<AFCDictionary> Data;
 
-			int ret = MobileDevice.DeviceImpl.FileInfoOpen(AFCCommsHandle, MobileDevice.StringToFileSystemRepresentation(path), out Data);
+			int ret = MobileDevice.DeviceImpl.FileInfoOpen(AFCCommsHandle, path, out Data);
             if ((ret == 0) && (Data.Handle != IntPtr.Zero))
             {
                 IntPtr pname;
@@ -795,7 +795,7 @@ namespace Manzana
                 // test for symbolic directory link
                 IntPtr hAFCDir = IntPtr.Zero;
 
-				if (directory = (MobileDevice.DeviceImpl.DirectoryOpen(AFCCommsHandle, MobileDevice.StringToFileSystemRepresentation(path), ref hAFCDir) == 0))
+				if (directory = (MobileDevice.DeviceImpl.DirectoryOpen(AFCCommsHandle, path, ref hAFCDir) == 0))
                 {
 					MobileDevice.DeviceImpl.DirectoryClose(AFCCommsHandle, hAFCDir);
                 }
@@ -823,7 +823,7 @@ namespace Manzana
         /// <returns>true if directory was created</returns>
         public bool CreateDirectory(string path)
         {
-			return !(MobileDevice.DeviceImpl.DirectoryCreate(AFCCommsHandle, MobileDevice.StringToFileSystemRepresentation(FullPath(CurrentDirectory, path))) != 0);
+			return !(MobileDevice.DeviceImpl.DirectoryCreate(AFCCommsHandle, FullPath(CurrentDirectory, path)) != 0);
         }
 
         /// <summary>
@@ -846,7 +846,7 @@ namespace Manzana
             string full_path = FullPath(CurrentDirectory, path);
             //full_path = "/private"; // bug test
 
-			int res = MobileDevice.DeviceImpl.DirectoryOpen(AFCCommsHandle, MobileDevice.StringToFileSystemRepresentation(full_path), ref hAFCDir);
+			int res = MobileDevice.DeviceImpl.DirectoryOpen(AFCCommsHandle, full_path, ref hAFCDir);
             if (res != 0)
             {
                 throw new Exception("Path does not exist: " + res.ToString());
@@ -877,8 +877,8 @@ namespace Manzana
         public bool Rename(string sourceName, string destName)
         {
 			return MobileDevice.DeviceImpl.RenamePath(AFCCommsHandle,
-				MobileDevice.StringToFileSystemRepresentation(FullPath(CurrentDirectory, sourceName)),
-				MobileDevice.StringToFileSystemRepresentation(FullPath(CurrentDirectory, destName))) == 0;
+				FullPath(CurrentDirectory, sourceName),
+				FullPath(CurrentDirectory, destName)) == 0;
         }
 
         /// <summary>
@@ -900,7 +900,7 @@ namespace Manzana
         {
             TypedPtr<AFCDictionary> data = IntPtr.Zero;
 
-			int ret = MobileDevice.DeviceImpl.FileInfoOpen(AFCCommsHandle, MobileDevice.StringToFileSystemRepresentation(path), out data);
+			int ret = MobileDevice.DeviceImpl.FileInfoOpen(AFCCommsHandle, path, out data);
             if (ret == 0)
             {
 				MobileDevice.DeviceImpl.KeyValueClose(data);
@@ -952,7 +952,7 @@ namespace Manzana
             string full_path = FullPath(CurrentDirectory, path);
             if (IsDirectory(full_path))
             {
-				MobileDevice.DeviceImpl.RemovePath(AFCCommsHandle, MobileDevice.StringToFileSystemRepresentation(full_path));
+				MobileDevice.DeviceImpl.RemovePath(AFCCommsHandle, full_path);
             }
         }
 
@@ -986,7 +986,7 @@ namespace Manzana
             string full_path = FullPath(CurrentDirectory, path);
             if (Exists(full_path))
             {
-				MobileDevice.DeviceImpl.RemovePath(AFCCommsHandle, MobileDevice.StringToFileSystemRepresentation(full_path));
+				MobileDevice.DeviceImpl.RemovePath(AFCCommsHandle, full_path);
             }
         }
         #endregion	// Filesystem
