Merging fix for mobile provision/certificate detection errors from 4.10

diff --git a/Engine/Source/Programs/IOS/iPhonePackager/MobileProvisionUtilities.cs b/Engine/Source/Programs/IOS/iPhonePackager/MobileProvisionUtilities.cs
--- a/Engine/Source/Programs/IOS/iPhonePackager/MobileProvisionUtilities.cs
+++ b/Engine/Source/Programs/IOS/iPhonePackager/MobileProvisionUtilities.cs
@@ -161,11 +161,25 @@ namespace iPhonePackager
 						continue;
 					}
 
-					bool bPassesDebugCheck = (!Config.bForDistribution || ((TestProvision.ProvisionedDeviceIDs.Count == 0) && !TestProvision.bDebug));
-					if (!bPassesDebugCheck)
+					if (Config.bForDistribution)
 					{
-						Program.LogVerbose("  .. Failed debugging check (mode={0}, get-task-allow={1}, #devices={2})", Config.bForDistribution, TestProvision.bDebug, TestProvision.ProvisionedDeviceIDs.Count);
-						continue;
+						// check to see if this is a distribution provision
+						bool bDistroProv = (TestProvision.ProvisionedDeviceIDs.Count == 0) && !TestProvision.bDebug;
+						if (!bDistroProv)
+						{
+							Program.LogVerbose("  .. Failed distribution check (mode={0}, get-task-allow={1}, #devices={2})", Config.bForDistribution, TestProvision.bDebug, TestProvision.ProvisionedDeviceIDs.Count);
+							continue;
+						}
+					}
+					else
+					{
+						// check to see if we pass the debug check for non-distribution
+						bool bPassesDebugCheck = TestProvision.bDebug;
+						if (!bPassesDebugCheck)
+						{
+							Program.LogVerbose("  .. Failed debugging check (mode={0}, get-task-allow={1}, #devices={2})", Config.bForDistribution, TestProvision.bDebug, TestProvision.ProvisionedDeviceIDs.Count);
+							continue;
+						}
 					}
 
 					// Check to see if the provision is in date
diff --git a/Engine/Source/Programs/IOS/iPhonePackager/iPhonePackager.cs b/Engine/Source/Programs/IOS/iPhonePackager/iPhonePackager.cs
--- a/Engine/Source/Programs/IOS/iPhonePackager/iPhonePackager.cs
+++ b/Engine/Source/Programs/IOS/iPhonePackager/iPhonePackager.cs
@@ -361,7 +361,7 @@ namespace iPhonePackager
 							if (Arguments.Length > ArgIndex + 1)
 							{
 								Config.Provision = Arguments [++ArgIndex];
-								Config.bProvision = true;
+								Config.bProvision = !String.IsNullOrEmpty(Config.Provision);
 							}
 							else
 							{
@@ -373,7 +373,7 @@ namespace iPhonePackager
 							if (Arguments.Length > ArgIndex + 1)
 							{
 								Config.Certificate = Arguments [++ArgIndex];
-								Config.bCert = true;
+								Config.bCert = !String.IsNullOrEmpty(Config.Certificate);
 							}
 							else
 							{
