fix for resource rules being deprecated with extreme prejudice in iOS 9
fix for requiring full screen for iOS 9

diff --git a/Engine/Source/Programs/IOS/iPhonePackager/CodeSigning.cs b/Engine/Source/Programs/IOS/iPhonePackager/CodeSigning.cs
--- a/Engine/Source/Programs/IOS/iPhonePackager/CodeSigning.cs
+++ b/Engine/Source/Programs/IOS/iPhonePackager/CodeSigning.cs
@@ -528,31 +528,14 @@ namespace iPhonePackager
 
 		protected virtual byte[] CreateCodeResourcesDirectory(string CFBundleExecutable)
 		{
-			// Verify that there is a rules plist
-			string CFBundleResourceSpecification;
-			if (!Info.GetString("CFBundleResourceSpecification", out CFBundleResourceSpecification))
-			{
-				throw new InvalidDataException("Info.plist must contain the key CFBundleResourceSpecification");
-			}
-
 			// Create a rules dict that includes (by wildcard) everything but Info.plist and the rules file
 			Dictionary<string, object> Rules = new Dictionary<string, object>();
 			Rules.Add(".*", true);
 			Rules.Add("^Info.plist", CreateOmittedResource(10));
-			Rules.Add(CFBundleResourceSpecification, CreateOmittedResource(100));
-
-			// Write the rules file out 
-			{
-				Utilities.PListHelper RulesPList = new Utilities.PListHelper();
-				RulesPList.AddKeyValuePair("rules", Rules);
-				string PListString = RulesPList.SaveToString();
-				FileSystem.WriteAllBytes(CFBundleResourceSpecification, Encoding.UTF8.GetBytes(PListString));
-			}
 
 			// Create the full list of files to exclude (some files get excluded by 'magic' even though they aren't listed as special by rules)
 			Dictionary<string, object> TrueExclusionList = new Dictionary<string, object>();
 			TrueExclusionList.Add("Info.plist", null);
-			TrueExclusionList.Add(CFBundleResourceSpecification, null);
 			TrueExclusionList.Add(CFBundleExecutable, null);
 			TrueExclusionList.Add("CodeResources", null);
 			TrueExclusionList.Add("_CodeSignature/CodeResources", null);
@@ -695,17 +678,6 @@ namespace iPhonePackager
 				throw new InvalidDataException("Info.plist must contain the key CFBundleIdentifier");
 			}
 
-			// Verify there is a resource rules file and make a dummy one if needed.
-			// If it's missing, CreateCodeResourceDirectory can't proceed (the Info.plist is already written to disk at that point)
-			if (!Info.HasKey("CFBundleResourceSpecification"))
-			{
-				// Couldn't find the key, create a dummy one
-				string CFBundleResourceSpecification = "CustomResourceRules.plist";
-				Info.SetString("CFBundleResourceSpecification", CFBundleResourceSpecification);
-
-				Program.Warning("Info.plist was missing the key CFBundleResourceSpecification, creating a new resource rules file '{0}'.", CFBundleResourceSpecification);
-			}
-
 			// Save the Info.plist out
 			byte[] RawInfoPList = Encoding.UTF8.GetBytes(Info.SaveToString());
 			Info.SetReadOnly(true);
diff --git a/Engine/Source/Programs/IOS/iPhonePackager/CompileTime.cs b/Engine/Source/Programs/IOS/iPhonePackager/CompileTime.cs
--- a/Engine/Source/Programs/IOS/iPhonePackager/CompileTime.cs
+++ b/Engine/Source/Programs/IOS/iPhonePackager/CompileTime.cs
@@ -78,10 +78,6 @@ namespace iPhonePackager
 					" -sdk " + ((Program.Architecture == "-simulator") ? "iphonesimulator" : "iphoneos");
 
 			// sign with the Distribution identity when packaging for distribution
-			if (!Config.bForDistribution)
-			{
-				CmdLine += String.Format(" CODE_SIGN_RESOURCE_RULES_PATH='{0}/CustomResourceRules.plist'", MacStagingRootDir);
-			}
 			if (Config.bUseRPCUtil)
 			{
 				CmdLine += String.Format(" CODE_SIGN_IDENTITY=\\\"{0}\\\"", Config.CodeSigningIdentity);
@@ -240,13 +236,6 @@ namespace iPhonePackager
 					Config.bForDistribution ? "false" : "true"));
 			}
 			
-			
-			// Copy the no sign resource rules file over
-			if (!Config.bForDistribution)
-			{
-				FileOperations.CopyRequiredFile(@"..\..\..\Build\IOS\XcodeSupportFiles\CustomResourceRules.plist", Path.Combine(Config.PCStagingRootDir, "CustomResourceRules.plist"));
-			}
-
 			// Copy the mobile provision file over
 			string CFBundleIdentifier = null;
 			Info.GetString("CFBundleIdentifier", out CFBundleIdentifier);
@@ -409,7 +398,7 @@ namespace iPhonePackager
 
 			case "strip":
 				Program.Log( " ... stripping" );
-				DisplayCommandLine = XcodeDeveloperDir + "Platforms/iPhoneOS.platform/Developer/usr/bin/strip '" + RemoteExecutablePath + "'";
+				DisplayCommandLine = "/usr/bin/xcrun strip '" + RemoteExecutablePath + "'";
 				CommandLine = "\"" + MacStagingRootDir + "\" " + DisplayCommandLine;
 				WorkingFolder = "\"" + MacStagingRootDir + "\"";
 				break;
diff --git a/Engine/Source/Programs/UnrealBuildTool/IOS/UEDeployIOS.cs b/Engine/Source/Programs/UnrealBuildTool/IOS/UEDeployIOS.cs
--- a/Engine/Source/Programs/UnrealBuildTool/IOS/UEDeployIOS.cs
+++ b/Engine/Source/Programs/UnrealBuildTool/IOS/UEDeployIOS.cs
@@ -307,6 +307,8 @@ namespace UnrealBuildTool
 			Text.AppendLine("\t<true/>");
 			Text.AppendLine("\t<key>UIStatusBarHidden</key>");
 			Text.AppendLine("\t<true/>");
+			Text.AppendLine("\t<key>UIRequiresFullScreen</key>");
+			Text.AppendLine("\t<true/>");
 			Text.AppendLine("\t<key>UIViewControllerBasedStatusBarAppearance</key>");
 			Text.AppendLine("\t<false/>");
 			Text.AppendLine("\t<key>UISupportedInterfaceOrientations</key>");
