Fix warning due to unused variable.

diff --git a/Engine/Source/Programs/IOS/iPhonePackager/GraphicalResignTool.cs b/Engine/Source/Programs/IOS/iPhonePackager/GraphicalResignTool.cs
--- a/Engine/Source/Programs/IOS/iPhonePackager/GraphicalResignTool.cs
+++ b/Engine/Source/Programs/IOS/iPhonePackager/GraphicalResignTool.cs
@@ -280,7 +280,7 @@ namespace iPhonePackager
 						SigningContext.CustomSigningCert = new X509Certificate2(CertificateEdit.Text, CertificatePassword, X509KeyStorageFlags.PersistKeySet | X509KeyStorageFlags.Exportable | X509KeyStorageFlags.MachineKeySet);
 						//Cert = new X509Certificate2(CertificateFilename, CertificatePassword, X509KeyStorageFlags.PersistKeySet | X509KeyStorageFlags.Exportable | X509KeyStorageFlags.MachineKeySet);
 					}
-					catch (System.Security.Cryptography.CryptographicException ex)
+					catch (System.Security.Cryptography.CryptographicException)
 					{
 						// Try once with a password
 						if (PasswordDialog.RequestPassword(out CertificatePassword))
