fix for one certificate check not using UTC

diff --git a/Engine/Source/Programs/IOS/iPhonePackager/CodeSigning.cs b/Engine/Source/Programs/IOS/iPhonePackager/CodeSigning.cs
--- a/Engine/Source/Programs/IOS/iPhonePackager/CodeSigning.cs
+++ b/Engine/Source/Programs/IOS/iPhonePackager/CodeSigning.cs
@@ -324,7 +324,7 @@ namespace iPhonePackager
 				X509Certificate2 Cert = FindCertificate(p);
 				if (Cert != null)
 				{
-					bValid = (Cert.NotBefore < Now) && (Cert.NotAfter > Now);
+					bValid = (Cert.NotBefore.ToUniversalTime() < Now) && (Cert.NotAfter.ToUniversalTime() > Now);
 				}
 				bool bPassesNameCheck = p.ApplicationIdentifier.Substring(p.ApplicationIdentifierPrefix.Length+1) == CFBundleIdentifier;
 				bool bPassesCompanyCheck = false;
@@ -458,10 +458,9 @@ namespace iPhonePackager
 					X509Certificate2 ValidInTimeCert = null;
 					foreach (X509Certificate2 TestCert in FoundCerts)
 					{
-						//@TODO: Pretty sure the certificate information from the library is in local time, not UTC and this works as expected, but it should be verified!
-						DateTime EffectiveDate = TestCert.NotBefore;
-						DateTime ExpirationDate = TestCert.NotAfter;
-						DateTime Now = DateTime.Now;
+						DateTime EffectiveDate = TestCert.NotBefore.ToUniversalTime();
+						DateTime ExpirationDate = TestCert.NotAfter.ToUniversalTime();
+						DateTime Now = DateTime.UtcNow;
 
 						bool bCertTimeIsValid = (EffectiveDate < Now) && (ExpirationDate > Now);
 
