Fix UAT compile on Mac and Linux by making prebuild step only run on Windows

diff --git a/Engine/Source/Programs/IOS/DeploymentInterface/DeploymentInterface.csproj b/Engine/Source/Programs/IOS/DeploymentInterface/DeploymentInterface.csproj
--- a/Engine/Source/Programs/IOS/DeploymentInterface/DeploymentInterface.csproj
+++ b/Engine/Source/Programs/IOS/DeploymentInterface/DeploymentInterface.csproj
@@ -63,7 +63,7 @@
   </ItemGroup>
   <Import Project="$(MSBuildToolsPath)\Microsoft.CSharp.targets" />
   <PropertyGroup>
-    <PreBuildEvent>StopDeploymentServer.bat
+    <PreBuildEvent Condition="'$(OS)' == 'Windows_NT' ">StopDeploymentServer.bat
 </PreBuildEvent>
   </PropertyGroup>
   <!-- To modify your build process, add your task inside one of the targets below and uncomment it. 
diff --git a/Engine/Source/Programs/IOS/DeploymentServer/DeploymentServer.csproj b/Engine/Source/Programs/IOS/DeploymentServer/DeploymentServer.csproj
--- a/Engine/Source/Programs/IOS/DeploymentServer/DeploymentServer.csproj
+++ b/Engine/Source/Programs/IOS/DeploymentServer/DeploymentServer.csproj
@@ -102,7 +102,7 @@
   </ItemGroup>
   <Import Project="$(MSBuildToolsPath)\Microsoft.CSharp.targets" />
   <PropertyGroup>
-    <PreBuildEvent>StopDeploymentServer.bat
+    <PreBuildEvent Condition="'$(OS)' == 'Windows_NT' ">StopDeploymentServer.bat
 </PreBuildEvent>
   </PropertyGroup>
   <!-- To modify your build process, add your task inside one of the targets below and uncomment it. 
diff --git a/Engine/Source/Programs/IOS/MobileDeviceInterface/MobileDeviceInterface.csproj b/Engine/Source/Programs/IOS/MobileDeviceInterface/MobileDeviceInterface.csproj
--- a/Engine/Source/Programs/IOS/MobileDeviceInterface/MobileDeviceInterface.csproj
+++ b/Engine/Source/Programs/IOS/MobileDeviceInterface/MobileDeviceInterface.csproj
@@ -115,7 +115,7 @@
   </ItemGroup>
   <Import Project="$(MSBuildToolsPath)\Microsoft.CSharp.targets" />
   <PropertyGroup>
-    <PreBuildEvent>StopDeploymentServer.bat
+    <PreBuildEvent Condition="'$(OS)' == 'Windows_NT' ">StopDeploymentServer.bat
 </PreBuildEvent>
   </PropertyGroup>
   <!-- To modify your build process, add your task inside one of the targets below and uncomment it. 
