Fix null pointer exception in iPhonePackager when a matching certificate cannot be found

diff --git a/Engine/Source/Programs/IOS/iPhonePackager/CompileTime.cs b/Engine/Source/Programs/IOS/iPhonePackager/CompileTime.cs
--- a/Engine/Source/Programs/IOS/iPhonePackager/CompileTime.cs
+++ b/Engine/Source/Programs/IOS/iPhonePackager/CompileTime.cs
@@ -229,7 +229,12 @@ namespace iPhonePackager
 			// export the signing certificate to a file
 			MobileProvision Provision = MobileProvisionParser.ParseFile(ProvisionWithPrefix);
 			var Certificate = CodeSignatureBuilder.FindCertificate(Provision);
-			byte[] Data = Certificate.Export(System.Security.Cryptography.X509Certificates.X509ContentType.Pkcs12, "A");
+            if (Certificate == null)
+            {
+                Program.Error("Failed to find a valid certificate");
+                return;
+            }
+            byte[] Data = Certificate.Export(System.Security.Cryptography.X509Certificates.X509ContentType.Pkcs12, "A");
 			File.WriteAllBytes(Config.Certificate, Data);
 		}
 
