Attempted fix for incorrect certificate errors when compiling remotely.

* Overloaded meaning of the -Certificate argument (ie. as the output path rather than as the desired certificate when selecting a provision) was causing certificate name check to fail. Now uses a different parameter to specify the output path.
* TVOS argument was not being passed when exporting the certificate, also causing the wrong provision/certificate pair to be selected.
* Add support for the -bundlename argument, though it shouldn't be necessary when invoked from UBT since the UUID discovered in a previous invocation of IPP is explicitly specified.

Since the provision UUID is already being specified to IPP, it should be able to identify the correct provision on phase -1 using this info.

diff --git a/Engine/Source/Programs/IOS/iPhonePackager/CompileTime.cs b/Engine/Source/Programs/IOS/iPhonePackager/CompileTime.cs
--- a/Engine/Source/Programs/IOS/iPhonePackager/CompileTime.cs
+++ b/Engine/Source/Programs/IOS/iPhonePackager/CompileTime.cs
@@ -9,6 +9,7 @@ using System.Diagnostics;
 using System.IO;
 using System.Net;
 using System.Runtime.InteropServices;
+using System.Security.Cryptography.X509Certificates;
 using System.Text;
 using System.Threading;
 using System.Windows.Forms;
@@ -213,29 +214,29 @@ namespace iPhonePackager
 		static public void ExportCertificate()
 		{
 			string ProvisionWithPrefix;
-			// this needs the provisioningUUID or it will fail
-			if(!FindMobileProvision("", out ProvisionWithPrefix, false))
+			if(!FindMobileProvision(Config.OverrideBundleName ?? "", out ProvisionWithPrefix, false))
 			{
 				Program.Error("Missing provision");
 				return;
 			}
 
-			if(Config.Certificate == null)
+			if(Config.OutputCertificate == null)
 			{
-				Program.Error("Missing -Certificate=... argument");
+				Program.Error("Missing -OutputCertificate=... argument");
 				return;
 			}
 
 			// export the signing certificate to a file
 			MobileProvision Provision = MobileProvisionParser.ParseFile(ProvisionWithPrefix);
-			var Certificate = CodeSignatureBuilder.FindCertificate(Provision);
+			X509Certificate2 Certificate = CodeSignatureBuilder.FindCertificate(Provision);
             if (Certificate == null)
             {
                 Program.Error("Failed to find a valid certificate");
                 return;
             }
-            byte[] Data = Certificate.Export(System.Security.Cryptography.X509Certificates.X509ContentType.Pkcs12, "A");
-			File.WriteAllBytes(Config.Certificate, Data);
+
+			byte[] Data = Certificate.Export(System.Security.Cryptography.X509Certificates.X509ContentType.Pkcs12, "A");
+			File.WriteAllBytes(Config.OutputCertificate, Data);
 		}
 
 		/// <summary>
diff --git a/Engine/Source/Programs/IOS/iPhonePackager/Config.cs b/Engine/Source/Programs/IOS/iPhonePackager/Config.cs
--- a/Engine/Source/Programs/IOS/iPhonePackager/Config.cs
+++ b/Engine/Source/Programs/IOS/iPhonePackager/Config.cs
@@ -365,6 +365,11 @@ namespace iPhonePackager
 		/// </summary>
 		public static string Certificate = "";
 
+		/// <summary>
+		/// Certificate to be output
+		/// </summary>
+		public static string OutputCertificate = null;
+
 		/// <summary>
 		/// An override server Mac name
 		/// </summary>
diff --git a/Engine/Source/Programs/IOS/iPhonePackager/iPhonePackager.cs b/Engine/Source/Programs/IOS/iPhonePackager/iPhonePackager.cs
--- a/Engine/Source/Programs/IOS/iPhonePackager/iPhonePackager.cs
+++ b/Engine/Source/Programs/IOS/iPhonePackager/iPhonePackager.cs
@@ -477,6 +477,17 @@ namespace iPhonePackager
 								return false;
 							}
 						}
+						else if (Arg == "-outputcertificate")
+						{
+							if (Arguments.Length > ArgIndex + 1)
+							{
+								Config.OutputCertificate = Arguments [++ArgIndex];
+							}
+							else
+							{
+								return false;
+							}
+						}
 					}
 					else
 					{
diff --git a/Engine/Source/Programs/UnrealBuildTool/ToolChain/RemoteMac.cs b/Engine/Source/Programs/UnrealBuildTool/ToolChain/RemoteMac.cs
--- a/Engine/Source/Programs/UnrealBuildTool/ToolChain/RemoteMac.cs
+++ b/Engine/Source/Programs/UnrealBuildTool/ToolChain/RemoteMac.cs
@@ -417,7 +417,11 @@ namespace UnrealBuildTool
 					{
 						Arguments.AppendFormat(" \"{0}\"", TargetDesc.ProjectFile.Directory);
 					}
-					Arguments.AppendFormat(" -certificate \"{0}\"", CertificateFile);
+					if(TargetDesc.Platform == UnrealTargetPlatform.TVOS)
+					{
+						Arguments.Append(" -tvos");
+					}
+					Arguments.AppendFormat(" -outputcertificate \"{0}\"", CertificateFile);
 					Arguments.AppendFormat(" -provisioninguuid {0}", ProvisioningData.MobileProvisionUUID);
 
 					ProcessStartInfo StartInfo = new ProcessStartInfo();
