UE-69280 - iOS Launch On Fails on Windows binary builds
UE-69887 - DeploymentServerLauncher.exe does not always quit

diff --git a/Engine/Source/Programs/IOS/DeploymentServer/Program.cs b/Engine/Source/Programs/IOS/DeploymentServer/Program.cs
--- a/Engine/Source/Programs/IOS/DeploymentServer/Program.cs
+++ b/Engine/Source/Programs/IOS/DeploymentServer/Program.cs
@@ -31,7 +31,7 @@ namespace DeploymentServer
 		static bool IsRunningCommand = false;
 		static bool IsStopping = false;
 		static int ClientCounter = 0;
-		static long TimeOut = 120000;
+		static long TimeOut = 30000;
 		static Stopwatch GlobalTimer = Stopwatch.StartNew();
 		static int ParentPID = 0;
 		static string TestStartPath = null;
@@ -280,36 +280,43 @@ namespace DeploymentServer
                             case "backupdocs":
                                 Console.SetOut(Writer);
                                 LastResult = DeploymentProxy.Deployer.BackupDocumentsDirectory(Bundle, FileList.Count > 0 ? FileList[0] : ".");
-                                break;
+								Writer.Flush();
+								break;
 
 							case "backup":
 								Console.SetOut(Writer);
 								LastResult = DeploymentProxy.Deployer.BackupFiles(Bundle, FileList.ToArray());
+								Writer.Flush();
 								break;
 
 							case "deploy":
 								Console.SetOut(Writer);
 								LastResult = DeploymentProxy.Deployer.InstallFilesOnDevice(Bundle, Manifest);
+								Writer.Flush();
 								break;
 
 							case "copyfile":
 								Console.SetOut(Writer);
 								LastResult = DeploymentProxy.Deployer.CopyFileToDevice(Bundle, FileList[0], FileList[1]);
+								Writer.Flush();
 								break;
 
 							case "install":
 								Console.SetOut(Writer);
 								LastResult = DeploymentProxy.Deployer.InstallIPAOnDevice(IpaPath);
+								Writer.Flush();
 								break;
 
 							case "enumerate":
 								Console.SetOut(Writer);
 								DeploymentProxy.Deployer.EnumerateConnectedDevices();
+								Writer.Flush();
 								break;
 
 							case "listdevices":
 								Console.SetOut(Writer);
 								DeploymentProxy.Deployer.ListDevices();
+								Writer.Flush();
 								break;
 
 							case "command":
@@ -349,6 +356,7 @@ namespace DeploymentServer
 										Writer.WriteLine("[command] Errors encountered while tunneling to device.");
 									}
 								}
+								Writer.Flush();
 								break;
 
 							case "forward":
@@ -375,6 +383,7 @@ namespace DeploymentServer
 									Writer.WriteLine("{0}\r{1}\r{2}", P.DeviceID, P.TCPPort, P.DevicePort);
 								}
 								Writer.WriteLine("");
+								Writer.Flush();
 								break;
 
 							case "listentodevice":
@@ -400,7 +409,6 @@ namespace DeploymentServer
 					}
 					finally
 					{
-						Writer.Flush();
 						Console.SetOut(ConsoleOld);
 						bCommandComplete = true;
 					}
@@ -652,9 +660,9 @@ namespace DeploymentServer
 			RemotingConfiguration.RegisterWellKnownServiceType(typeof(DeploymentProxy), URI, WellKnownObjectMode.Singleton);
 		}
 
-		protected static void ParseServerParam(string[] Arguments)
+		protected static void ParseServerParam(List<string> Arguments)
 		{
-			if (Arguments.Length > 2)
+			if (Arguments.Count > 2)
 			{
 				TestStartPath = Arguments[2];
 			}
@@ -662,9 +670,9 @@ namespace DeploymentServer
 			{
 				TestStartPath = GetDeploymentServerPath();
 			}
-			if (Arguments.Length > 3)
+			if (Arguments.Count > 3)
 			{
-				for (int ArgIndex = 3; ArgIndex < Arguments.Length; ArgIndex++)
+				for (int ArgIndex = 3; ArgIndex < Arguments.Count; ArgIndex++)
 				{
 					string Arg = Arguments[ArgIndex].ToLowerInvariant();
 					if (Arg.StartsWith("-"))
@@ -673,7 +681,7 @@ namespace DeploymentServer
 						{
 							case "-timeout":
 								{
-									if (Arguments.Length > ArgIndex + 1)
+									if (Arguments.Count > ArgIndex + 1)
 									{
 										long ArgTime = TimeOut;
 										long.TryParse(Arguments[++ArgIndex], out ArgTime);
@@ -738,11 +746,20 @@ namespace DeploymentServer
 				}
 				Server = new TcpListener(IPAddress.Any, Port);
 				Server.Start();
-				
 
-				ParseServerParam(Args);
+				string CommandLine = "";
+				foreach (string Arg in Args)
+				{
+					CommandLine += Arg + " ";
+				}
+				List<string> Arguments = Regex.Matches(CommandLine, @"[\""].+?[\""]|[^ ]+")
+												.Cast<Match>()
+												.Select(m => m.Value)
+												.ToList();
+				ParseServerParam(Arguments);
 				Console.WriteLine(string.Format("Deployment Server listening to port {0}", Port.ToString()));
 				Console.WriteLine(string.Format("Deployment Server inactivity timeout {0}", TimeOut.ToString()));
+				Console.WriteLine(string.Format("Deployment Server starting from {0}", TestStartPath));
 				Console.WriteLine("---------------------------------------------------------");
 
 				// Processing commands
diff --git a/Engine/Source/Programs/IOS/DeploymentServerLauncher/Launcher.cs b/Engine/Source/Programs/IOS/DeploymentServerLauncher/Launcher.cs
--- a/Engine/Source/Programs/IOS/DeploymentServerLauncher/Launcher.cs
+++ b/Engine/Source/Programs/IOS/DeploymentServerLauncher/Launcher.cs
@@ -26,7 +26,7 @@ namespace DeploymentServerLauncher
 				NewProcess.StartInfo.WorkingDirectory = WorkingDirectory;
 				NewProcess.StartInfo.WorkingDirectory.TrimEnd('/');
 				NewProcess.StartInfo.FileName = NewProcess.StartInfo.WorkingDirectory + "/../../../Build/BatchFiles/Mac/RunMono.sh";
-				NewProcess.StartInfo.Arguments = "\"" + NewProcess.StartInfo.WorkingDirectory + "/DeploymentServer.exe\" server " + Process.GetCurrentProcess().Id.ToString() + " " + NewProcess.StartInfo.WorkingDirectory;
+				NewProcess.StartInfo.Arguments = "\"" + NewProcess.StartInfo.WorkingDirectory + "/DeploymentServer.exe\" server " + Process.GetCurrentProcess().Id.ToString() + " \"" + NewProcess.StartInfo.WorkingDirectory + "\"";
 
 				NewProcess.StartInfo.WindowStyle = ProcessWindowStyle.Hidden;
 				NewProcess.StartInfo.UseShellExecute = true;
@@ -47,7 +47,7 @@ namespace DeploymentServerLauncher
 				Domaininfo.ConfigurationFile = Path.GetDirectoryName(Assembly.GetExecutingAssembly().Location) + "\\DeploymentServer.exe.config";
 				WorkingDirectory.TrimEnd('\\');
 				UATExecutable = WorkingDirectory + "\\DeploymentServer.exe";
-				CommandLine = "server " + Process.GetCurrentProcess().Id.ToString() + " " + WorkingDirectory;
+				CommandLine = "server " + Process.GetCurrentProcess().Id.ToString() + " \"" + WorkingDirectory + "\"";
 				Domaininfo.ShadowCopyFiles = "true";
 
 				// Create application domain setup information.
