UE-69786 - Installed Build fails to Build Tools CS - error: Command 'StopDeploymentServer.bat

diff --git a/Engine/Source/Programs/IOS/DeploymentServerLauncher/DeploymentServerLauncher.csproj b/Engine/Source/Programs/IOS/DeploymentServerLauncher/DeploymentServerLauncher.csproj
--- a/Engine/Source/Programs/IOS/DeploymentServerLauncher/DeploymentServerLauncher.csproj
+++ b/Engine/Source/Programs/IOS/DeploymentServerLauncher/DeploymentServerLauncher.csproj
@@ -76,8 +76,9 @@
   </ItemGroup>
   <Import Project="$(MSBuildToolsPath)\Microsoft.CSharp.targets" />
   <PropertyGroup>
-    <PreBuildEvent>StopDeploymentServer.bat
-</PreBuildEvent>
+    <PreBuildEvent Condition=" '$(OS)' == 'Windows_NT' ">
+	StopDeploymentServer.bat
+    </PreBuildEvent>
   </PropertyGroup>
   <!-- To modify your build process, add your task inside one of the targets below and uncomment it. 
        Other similar extension points exist, see Microsoft.Common.targets.
