UE-66204 iOS and tvOS errors in Application Loader to Apple: 'Invalid Code Signing Entitlements' with iCloud

#4.22
#iOS

Fixed a problem with the order of values in the  com.apple.developer.icloud-container-environment array

Also modified parsing com.apple.developer.ubiquity-container-identifiers and com.apple.developer.ubiquity-kvstore-identifier in order to preserve the team id from the mobileprovisioning

Note:iPhonePacvkager needs to be recompiled + binaries uploaded on //UE4/Release-4.22/Engine/Binaries/DotNET/IOS

diff --git a/Engine/Source/Programs/IOS/iPhonePackager/MobileProvisionUtilities.cs b/Engine/Source/Programs/IOS/iPhonePackager/MobileProvisionUtilities.cs
--- a/Engine/Source/Programs/IOS/iPhonePackager/MobileProvisionUtilities.cs
+++ b/Engine/Source/Programs/IOS/iPhonePackager/MobileProvisionUtilities.cs
@@ -341,37 +341,47 @@ namespace iPhonePackager
 				// For distribution builds, the entitlements from mobileprovisioning have a modified syntax
 				if (Config.bForDistribution)
 				{
+					// remove the wildcards from the ubiquity-kvstore-identifier string
 					if (XCentPList.HasKey("com.apple.developer.ubiquity-kvstore-identifier"))
 					{
 						string UbiquityKvstoreString;
 						XCentPList.GetString("com.apple.developer.ubiquity-kvstore-identifier", out UbiquityKvstoreString);
 
-						if (UbiquityKvstoreString.Contains("*"))
+						int DotPosition = UbiquityKvstoreString.LastIndexOf("*");
+						if (DotPosition >= 0)
 						{
-							string NewUbiquityKvstoreIdentifier = String.Format("{0}.{1}", TeamIdentifier, CFBundleIdentifier);
+							string TeamPrefix = DotPosition > 1 ? UbiquityKvstoreString.Substring(0, DotPosition - 1) : TeamIdentifier;
+							string NewUbiquityKvstoreIdentifier = String.Format("{0}.{1}", TeamPrefix, CFBundleIdentifier);
 							XCentPList.SetValueForKey("com.apple.developer.ubiquity-kvstore-identifier", NewUbiquityKvstoreIdentifier);
 						}
 					}
+
+					// remove the wildcards from the ubiquity-container-identifiers array
 					if (XCentPList.HasKey("com.apple.developer.ubiquity-container-identifiers"))
 					{
 						List<string> UbiquityContainerIdentifiersGroups = XCentPList.GetArray("com.apple.developer.ubiquity-container-identifiers", "string");
 
-						if (UbiquityContainerIdentifiersGroups.Count == 0 || UbiquityContainerIdentifiersGroups[0].Contains("*"))
+						for (int i = 0; i < UbiquityContainerIdentifiersGroups.Count; i++)
 						{
-							UbiquityContainerIdentifiersGroups.Clear();
-
-							string UbiquityContainerString;
-							XCentPList.GetString("com.apple.developer.ubiquity-container-identifiers", out UbiquityContainerString);
-
-							if (UbiquityContainerString.Contains("*"))
+							int DotPosition = UbiquityContainerIdentifiersGroups[i].LastIndexOf("*");
+							if (DotPosition >= 0)
 							{
-								string NewUbiquityContainerIdentifier = String.Format("{0}.{1}", TeamIdentifier, CFBundleIdentifier);
-								UbiquityContainerIdentifiersGroups.Add(NewUbiquityContainerIdentifier);
+								string TeamPrefix = DotPosition > 1 ? UbiquityContainerIdentifiersGroups[i].Substring(0, DotPosition - 1) : TeamIdentifier;
+								string NewUbiquityContainerIdentifier = String.Format("{0}.{1}", TeamPrefix, CFBundleIdentifier);
+								UbiquityContainerIdentifiersGroups[i] = NewUbiquityContainerIdentifier;
 							}
 						}
 
+						if (UbiquityContainerIdentifiersGroups.Count == 0)
+						{
+							string NewUbiquityKvstoreIdentifier = String.Format("{0}.{1}", TeamIdentifier, CFBundleIdentifier);
+							UbiquityContainerIdentifiersGroups.Add(NewUbiquityKvstoreIdentifier);
+						}
+
 						XCentPList.SetValueForKey("com.apple.developer.ubiquity-container-identifiers", UbiquityContainerIdentifiersGroups);
 					}
+
+					// remove the wildcards from the developer.associated-domains array or string
 					if (XCentPList.HasKey("com.apple.developer.associated-domains"))
 					{
 						List<string> AssociatedDomainsGroup = XCentPList.GetArray("com.apple.developer.associated-domains", "string");
@@ -390,16 +400,17 @@ namespace iPhonePackager
 
 					// remove development keys - generated when the cloudkit container is in development mode
 					XCentPList.RemoveKeyValue("com.apple.developer.icloud-container-development-container-identifiers");
-					if (XCentPList.HasKey("com.apple.developer.icloud-container-environment"))
-					{
-						List<string> ContainerEnvironmentGroup = XCentPList.GetArray("com.apple.developer.icloud-container-environment", "string");
+				}
 
-						if (ContainerEnvironmentGroup.Count == 0 || ContainerEnvironmentGroup[0].Equals("Development"))
-						{
-							ContainerEnvironmentGroup.Clear();
-							ContainerEnvironmentGroup.Add("Production");
-						}
+				// set the icloud-container-environment accordign to project settings
+				if (XCentPList.HasKey("com.apple.developer.icloud-container-environment"))
+				{
+					List<string> ContainerEnvironmentGroup = XCentPList.GetArray("com.apple.developer.icloud-container-environment", "string");
 
+					if (ContainerEnvironmentGroup.Count != 0)
+					{
+						ContainerEnvironmentGroup.Clear();
+						ContainerEnvironmentGroup.Add(Config.bForDistribution ? "Production" : "Development");
 						XCentPList.SetValueForKey("com.apple.developer.icloud-container-environment", ContainerEnvironmentGroup);
 					}
 				}
