- Fixed IPP with HiDPI monitors

diff --git a/Engine/Source/Programs/IOS/iPhonePackager/iPhonePackager.cs b/Engine/Source/Programs/IOS/iPhonePackager/iPhonePackager.cs
--- a/Engine/Source/Programs/IOS/iPhonePackager/iPhonePackager.cs
+++ b/Engine/Source/Programs/IOS/iPhonePackager/iPhonePackager.cs
@@ -682,11 +682,19 @@ namespace iPhonePackager
 			}
 			bVisualsStarted = true;
 
+			// make the form look good on modern displays!
+			if (Environment.OSVersion.Version.Major >= 6)
+			{
+				SetProcessDPIAware();
+			}
+
 			Application.EnableVisualStyles();
 			Application.SetCompatibleTextRenderingDefault(false);
 
 			ProgressDialog = new SlowProgressDialog();
 		}
+		[System.Runtime.InteropServices.DllImport("user32.dll")]
+		private static extern bool SetProcessDPIAware();
 
 		static void RunInVisualMode(CreateFormDelegate Work)
 		{
