UE-72058 tvOS: Project Packaged with CloudKit Support Crashes on Launch Due to Malformed Value in the Entitlement

#4.22
#tvOS

Two issues found and fixed:
- com.apple.developer.icloud-services should have CloudDocuments ONLY on iOS
- com.apple.developer.icloud-container-environment should by a string not an array (unlike in the mobile provisioning)

diff --git a/Engine/Source/Programs/IOS/iPhonePackager/MobileProvisionUtilities.cs b/Engine/Source/Programs/IOS/iPhonePackager/MobileProvisionUtilities.cs
--- a/Engine/Source/Programs/IOS/iPhonePackager/MobileProvisionUtilities.cs
+++ b/Engine/Source/Programs/IOS/iPhonePackager/MobileProvisionUtilities.cs
@@ -318,23 +318,14 @@ namespace iPhonePackager
 			// must have CloudKit and CloudDocuments for com.apple.developer.icloud-services
 			// otherwise the game will not be listed in the Settings->iCloud apps menu on the device
 			{
-				if (XCentPList.HasKey("com.apple.developer.icloud-services"))
+				// iOS only
+				if (Platform == "IOS" && XCentPList.HasKey("com.apple.developer.icloud-services"))
 				{
 					List<string> ServicesGroups = XCentPList.GetArray("com.apple.developer.icloud-services", "string");
+					ServicesGroups.Clear();
 
-					if (ServicesGroups.Count == 0 || !ServicesGroups[0].Contains("*"))
-					{
-						ServicesGroups.Clear();
-
-						string ServicesString;
-						XCentPList.GetString("com.apple.developer.icloud-services", out ServicesString);
-
-						if (ServicesString.Contains("*"))
-						{
-							ServicesGroups.Add("CloudKit");
-							ServicesGroups.Add("CloudDocuments");
-						}
-					}
+					ServicesGroups.Add("CloudKit");
+					ServicesGroups.Add("CloudDocuments");
 					XCentPList.SetValueForKey("com.apple.developer.icloud-services", ServicesGroups);
 				}
 
@@ -402,7 +393,7 @@ namespace iPhonePackager
 					XCentPList.RemoveKeyValue("com.apple.developer.icloud-container-development-container-identifiers");
 				}
 
-				// set the icloud-container-environment accordign to project settings
+				// set the icloud-container-environment according to the project settings
 				if (XCentPList.HasKey("com.apple.developer.icloud-container-environment"))
 				{
 					List<string> ContainerEnvironmentGroup = XCentPList.GetArray("com.apple.developer.icloud-container-environment", "string");
@@ -410,8 +401,10 @@ namespace iPhonePackager
 					if (ContainerEnvironmentGroup.Count != 0)
 					{
 						ContainerEnvironmentGroup.Clear();
-						ContainerEnvironmentGroup.Add(Config.bForDistribution ? "Production" : "Development");
-						XCentPList.SetValueForKey("com.apple.developer.icloud-container-environment", ContainerEnvironmentGroup);
+						
+						// The new value is a string, not an array
+						string NewContainerEnvironment = Config.bForDistribution ? "Production" : "Development";
+						XCentPList.SetValueForKey("com.apple.developer.icloud-container-environment", NewContainerEnvironment);
 					}
 				}
 			}
