UE-72446 - Deployment server crashes in DeploymentServer.Program.ClientLoop

diff --git a/Engine/Source/Programs/IOS/DeploymentServer/Program.cs b/Engine/Source/Programs/IOS/DeploymentServer/Program.cs
--- a/Engine/Source/Programs/IOS/DeploymentServer/Program.cs
+++ b/Engine/Source/Programs/IOS/DeploymentServer/Program.cs
@@ -273,7 +273,6 @@ namespace DeploymentServer
 										break;
 									}
 								}
-								ForceKillProcesses();
 								break;
 
                             case "backupdocuments":
@@ -398,11 +397,14 @@ namespace DeploymentServer
 								break;
 						}
 					}
-					
 					catch (IOException)
 					{
 						// we expect this to happen so we don't log it
 					}
+					catch (ThreadAbortException)
+					{
+						// we expect this to happen so we don't log it
+					}
 					catch (Exception e)
 					{
 						Console.SetOut(ConsoleOld);
@@ -428,6 +430,10 @@ namespace DeploymentServer
 						{
 							CoreFoundationRunLoop.RunLoopRunInMode(CoreFoundationRunLoop.kCFRunLoopDefaultMode(), 1.0, 0);
 						}
+						if (runLoop.ThreadState == System.Threading.ThreadState.Running)
+						{
+							runLoop.Abort();
+						}
 					}
 				}
 				catch (Exception e)
@@ -631,20 +637,53 @@ namespace DeploymentServer
 			}
 		}
 
-		static TcpClient IsServiceRegistered()
+		static TcpClient GetServiceClient(string LocalCommand)
 		{
 			TcpClient Client = null;
-			try
-			{
-				Client = new TcpClient("localhost", Port);
-			}
-			catch
+			int RetryCount = 10;
+			while (RetryCount > 0 && Client == null)
 			{
-				if (Client != null)
+				if (!IsServerMutexPresent())
 				{
-					Client.Close();
+					if (LocalCommand == "stop")
+					{
+						Console.WriteLine("Deployment Server not running ...");
+						return null;
+					}
+					Console.WriteLine("Creating Deployment Server ...");
+					CreateDeploymentServerProcess();
+				}
+				else
+				{
+					if (RetryCount == 5)
+					{
+						ForceKillProcesses();
+						Console.WriteLine("Creating Deployment Server ...");
+						CreateDeploymentServerProcess();
+					}
+				}
+				try
+				{
+					if (Client == null)
+					{
+						Client = new TcpClient("localhost", Port); //The client gets here
+					}
+				}
+				catch (Exception e)
+				{
+					//Console.WriteLine("Retrying count {0} {1}", RetryCount.ToString(), e.Message);
+					RetryCount--;
+					if (Client != null)
+					{
+						Client.Close();
+					}
+					Client = null;
+					if (RetryCount <= 0)
+					{
+						throw (e);
+					}
+					System.Threading.Thread.Sleep(750);
 				}
-				Client = null;
 			}
 
 			return Client;
@@ -704,31 +743,39 @@ namespace DeploymentServer
 			}
 		}
 
+		private static String ServerMutexName = "Global\\DeploymentServer_Mutex_SERVERINSTANCE";
+
+		static bool IsServerMutexPresent()
+		{
+			Mutex DSBlockMutex = null;
+			bool ret = Mutex.TryOpenExisting(ServerMutexName, out DSBlockMutex);
+			if (ret)
+			{
+				DSBlockMutex.Close();
+			}
+			return ret;
+		}
+
 		/**
 		 *	Main Server Loop
 		 */
-		static void ServerLoop(TcpClient IsServiceRunning, string[] Args)
+		static void ServerLoop(string[] Args)
 		{
 			Program.ExitCode = 0;
-
-			if (IsServiceRunning != null)
+			if (IsServerMutexPresent())
 				return;
 
 			bool bCreatedMutex = false;
-			String MutexName = "Global\\DeploymentServer_Mutex_SERVERINSTANCE";
 			Mutex DSBlockMutex = null;
 
 			try
 			{
-				DSBlockMutex = new Mutex(true, MutexName, out bCreatedMutex);
-				if (!bCreatedMutex)
-				{
-					// running is not allowed (perhaps a another server instance is running, but on a different port or is not responding to a connection request)
-					return;
-				}
+				DSBlockMutex = new Mutex(false, ServerMutexName, out bCreatedMutex);
+				DSBlockMutex.WaitOne();
 			}
 			catch
 			{
+				Console.WriteLine("Exception creating Deployment Server Mutex. ");
 				return;
 			}
 
@@ -831,9 +878,27 @@ namespace DeploymentServer
 			{
 				if (DSBlockMutex != null)
 				{
-					DSBlockMutex.ReleaseMutex();
-					DSBlockMutex.Dispose();
-					DSBlockMutex = null;
+					try
+					{
+						DSBlockMutex.ReleaseMutex();
+						DSBlockMutex.Close();
+						DSBlockMutex = null;
+					}
+					catch (AbandonedMutexException)
+					{
+						//This catch is included to insure the program keeps running in the event this exception occurs.
+						Console.WriteLine("Deployment Server Mutex abandoned.");
+					}
+					catch (ApplicationException)
+					{
+						//This catch is included to insure the program keeps running in the event this exception occurs.
+						Console.WriteLine("Deployment Server Mutex abandoned 1.");
+					}
+					catch (SynchronizationLockException)
+					{
+						//This catch is included to insure the program keeps running in the event this exception occurs.
+						Console.WriteLine("Deployment Server Mutex abandoned. 2");
+					}
 				}
 				if (ProcessClient != null)
 				{
@@ -901,60 +966,26 @@ namespace DeploymentServer
 		/**
 		 * Main Client loop
 		 */
-		static void ClientLoop(TcpClient IsServiceRunning, string[] Args, string LocalCommand)
+		static void ClientLoop(string[] Args, string LocalCommand)
 		{
-			if (IsServiceRunning == null)
+			// on mac we only start one local instance due to mono limitations
+			if (Environment.OSVersion.Platform == PlatformID.MacOSX || Environment.OSVersion.Platform == PlatformID.Unix || Args[0].Equals("-standalone"))
 			{
-				if (Args.Length < 1 || LocalCommand == "stop")
-				{
-					Console.WriteLine("Deployment Server not running ...");
-					ForceKillProcesses();
-					return;
-				}
-				// on mac we only start one local instance due to mono limitations
-				if (Environment.OSVersion.Platform == PlatformID.MacOSX || Environment.OSVersion.Platform == PlatformID.Unix || Args[0].Equals("-standalone"))
-				{
-					RunLocalInstance(Args);
-					return;
-				}
-				else
-				{
-					CreateDeploymentServerProcess();
-				}
+				RunLocalInstance(Args);
+				return;
 			}
-			// Parse the command
-			TcpClient Client = IsServiceRunning;
+
+			TcpClient Client = GetServiceClient(LocalCommand);
+
 			try
 			{
-				int RetryCount = 3;
-				while (RetryCount > 0 && Client == null)
-				{
-					try
-					{
-						if (Client == null)
-						{
-							Client = new TcpClient("localhost", Port); //The client gets here
-						}
-					}
-					catch (Exception e)
-					{
-						RetryCount--;
-						Client = null;
-						if (RetryCount <= 0)
-						{
-							throw (e);
-						}
-						System.Threading.Thread.Sleep(500);
-					}
-				}
 				StreamReader clientIn = new StreamReader(Client.GetStream());
 				StreamWriter clientOut = new StreamWriter(Client.GetStream());
 				{
 					string Response = clientIn.ReadLine();
 					if (TcpClientInfo.NeedsVersionCheck(LocalCommand))
 					{
-
-						if (!Response.Equals("[DSDIR]" + GetDeploymentServerPath(), StringComparison.InvariantCultureIgnoreCase))
+						if (!Response.Equals("[DSDIR]\"" + GetDeploymentServerPath() + "\"", StringComparison.InvariantCultureIgnoreCase))
 						{
 							Console.WriteLine("Wrong server running, restarting the server ...");
 							clientOut.Write("stop");
@@ -971,16 +1002,13 @@ namespace DeploymentServer
 								{
 									Console.WriteLine(Response);
 								}
-								ForceKillProcesses();
 							}
 							if (Client != null)
 							{
 								Client.Close();
 							}
-							Client = null;
-							CreateDeploymentServerProcess();
-							IsServiceRunning = IsServiceRegistered();
-							Client = IsServiceRunning;
+							Thread.Sleep(2500);
+							Client = GetServiceClient("");
 							clientIn = new StreamReader(Client.GetStream());
 							clientOut = new StreamWriter(Client.GetStream());
 						}
@@ -1074,29 +1102,6 @@ namespace DeploymentServer
 
 				return 0;
 			}
-			try
-			{
-				if (LocalCommand != "stop" && LocalCommand != "-iphonepackager")
-				{
-					bool bCreatedMutex = false;
-					String MutexName = "Global\\DeploymentServer_Mutex_RestartNotAllowed";
-					Mutex DSBlockMutex = new Mutex(true, MutexName, out bCreatedMutex);
-					if (!bCreatedMutex)
-					{
-						// running is not allowed (perhaps a build command is in progress)
-						Program.ExitCode = 1;
-						return 1;
-					}
-					DSBlockMutex.ReleaseMutex();
-					DSBlockMutex.Dispose();
-					DSBlockMutex = null;
-				}
-			}
-			catch
-			{
-				Program.ExitCode = 0;
-				return 0;
-			}
 			int.TryParse(ConfigurationManager.AppSettings["DSPort"], out Port);
 			if (Port < 1 || Port > 65535)
 			{
@@ -1106,8 +1111,7 @@ namespace DeploymentServer
 			// parrent ID not needed anymore
 			if (Args[0].Equals("server"))
 			{
-				TcpClient IsServiceRunning = IsServiceRegistered();
-				ServerLoop(IsServiceRunning, Args);
+				ServerLoop(Args);
 			}
 			else if (Args[0].Equals("-iphonepackager"))
 			{
@@ -1116,8 +1120,7 @@ namespace DeploymentServer
 			}
 			else
 			{
-				TcpClient IsServiceRunning = IsServiceRegistered();
-				ClientLoop(IsServiceRunning, Args, LocalCommand);
+				ClientLoop(Args, LocalCommand);
 			}
 
 			Environment.ExitCode = Program.ExitCode;
@@ -1135,7 +1138,7 @@ namespace DeploymentServer
 				{
 					if (i < Commands.Length - 1 || Commands[i].EndsWith("\r"))
 					{
-						List<string> Arguments = Regex.Matches(Commands[i], @"[\""].+?[\""]|[^ ]+")
+						List<string> Arguments = Regex.Matches(Commands[i].TrimEnd('\r'), @"[\""].+?[\""]|[^ ]+")
 												.Cast<Match>()
 												.Select(m => m.Value)
 												.ToList();
@@ -1166,7 +1169,8 @@ namespace DeploymentServer
 					string ClientIP = ((IPEndPoint)Client.Client.RemoteEndPoint).Address.ToString();
 					Console.WriteLine("Client [{0}] IP:{1} connected.", localID, ClientIP);
 
-					TextWriter Writer = new StreamWriter(ClStream);
+					StreamWriter Writer = new StreamWriter(ClStream);
+					Writer.AutoFlush = true;
 					
 					Writer.WriteLine("[DSDIR]" + TestStartPath);
 					Writer.Flush();
@@ -1174,6 +1178,7 @@ namespace DeploymentServer
 					Byte[] Buffer = new Byte[2048];
 					string Unprocessed = "";
 					ClientInfo = new TcpClientInfo();
+					bool IsCommandStarted = false;
 					
 					while (true)
 					{
@@ -1189,6 +1194,7 @@ namespace DeploymentServer
 								{
 									LastCommand = ClientInfo.LastCommand;
 									ClientInfo.RunCommand(Writer);
+									IsCommandStarted = true;
 								}
 								catch
 								{
@@ -1254,6 +1260,11 @@ namespace DeploymentServer
 								Client.Client.Blocking = BlockingState;
 							}
 						}
+						else if (IsCommandStarted)
+						{
+							// long running command has stopped
+							break;
+						}
 						if (IsStopping)
 						{
 							try
