UE-73030 - DeploymentServerLauncher crash with multiple editor sessions

diff --git a/Engine/Source/Programs/IOS/DeploymentServer/Program.cs b/Engine/Source/Programs/IOS/DeploymentServer/Program.cs
--- a/Engine/Source/Programs/IOS/DeploymentServer/Program.cs
+++ b/Engine/Source/Programs/IOS/DeploymentServer/Program.cs
@@ -36,6 +36,12 @@ namespace DeploymentServer
 		static int ParentPID = 0;
 		static string TestStartPath = null;
 
+		static TextWriter LocalConsole;
+		static public void LocalLog(string s)
+		{
+			LocalConsole.WriteLine(s);
+		}
+
 		class TCPPortForwarding
 		{
 			public String DeviceID;
@@ -215,7 +221,7 @@ namespace DeploymentServer
 										if (ArgTime > 0)
 										{
 											TimeOut = ArgTime;
-											Console.WriteLine(string.Format("Deployment Server timeout set to {0} (remote)", TimeOut.ToString()));
+											LocalLog(string.Format("Deployment Server timeout set to {0} (remote)", TimeOut.ToString()));
 										}
 									}
 									break;
@@ -257,7 +263,7 @@ namespace DeploymentServer
 						{
 							case "stop":
 								Console.SetOut(Writer);
-								Console.WriteLine("Deployment Server Stopping ...");
+								//LocalLog("Deployment Server Stopping ...");
 								IsStopping = true;
 								int StopTimeout = 12;
 								while (ClientCounter > 1) // wait for other threads to stop so the client requesting to stop to block until safely stopped
@@ -267,7 +273,7 @@ namespace DeploymentServer
 									StopTimeout--;
 									if (StopTimeout <= 0)
 									{
-										Console.WriteLine("Deployment Server Forced Stopping ...");
+										//LocalLog("Deployment Server Forced Stopping ...");
 										ClientCounter = 0;
 										TCPForwarding.Clear();
 										break;
@@ -321,12 +327,12 @@ namespace DeploymentServer
 							case "command":
 								if (Device.Length < 5)
 								{
-									Console.WriteLine("Device ID not present.");
+									LocalLog("Device ID not present.");
 									Writer.WriteLine("[command] Device ID not present.");
 								}
 								else if (Param1.Length < 1)
 								{
-									Console.WriteLine("Parameter not present.");
+									LocalLog("Parameter not present.");
 									Writer.WriteLine("[command] Device ID not present.");
 								}
 								else
@@ -340,18 +346,18 @@ namespace DeploymentServer
 											int Ret = targetDevice.TunnelData(Param1, TCPService);
 											targetDevice.CloseTunnel(TCPService);
 
-											Console.WriteLine("[UE4][command] Sent '{0}' bytes. ({1})", Ret, Param1);
+											LocalLog("[UE4][command] Sent '" + Ret.ToString() + "' bytes. (" + Param1.ToString() + ")");
 											Writer.WriteLine("[UE4][command] Sent '{0}' bytes. ({1})", Ret, Param1);
 										}
 										else
 										{
-											Console.WriteLine("[UE4][command] Device '{0}' not detected. ({1})", Device, Param1);
+											LocalLog("[UE4][command] Sent '" + Device.ToString() + "' bytes. (" + Param1.ToString() + ")");
 											Writer.WriteLine("[UE4][command] Device '{0}' not detected. ({1})", Device, Param1);
 										}
 									}
 									catch
 									{
-										Console.WriteLine("Errors encountered while tunneling to device.");
+										LocalLog("Errors encountered while tunneling to device.");
 										Writer.WriteLine("[command] Errors encountered while tunneling to device.");
 									}
 								}
@@ -361,15 +367,15 @@ namespace DeploymentServer
 							case "forward":
 								if (Device.Length < 5)
 								{
-									Console.WriteLine("Device ID not present.");
+									LocalLog("Device ID not present.");
 								}
 								else if (Param1.Length < 1)
 								{
-									Console.WriteLine("Start TCP port not present.");
+									LocalLog("Start TCP port not present.");
 								}
 								else if (Param2.Length < 1)
 								{
-									Console.WriteLine("Destination TCP port not present.");
+									LocalLog("Destination TCP port not present.");
 								}
 								else
 								{
@@ -388,7 +394,7 @@ namespace DeploymentServer
 							case "listentodevice":
 								if (Device.Length < 5)
 								{
-									Console.WriteLine("Device ID not present.");
+									LocalLog("Device ID not present.");
 								}
 								else
 								{
@@ -410,7 +416,7 @@ namespace DeploymentServer
 						Console.SetOut(ConsoleOld);
 						if (Command != "stop")
 						{
-							Console.WriteLine("Exception: {0}", e);
+							LocalLog("Exception: " + e.ToString());
 						}
 						LastResult = false;
 					}
@@ -441,7 +447,7 @@ namespace DeploymentServer
 					
 					if (Command != "stop")
 					{
-						Console.WriteLine("Exception: {0}", e);
+						LocalLog("Exception: " + e.ToString());
 					}
 					bCommandComplete = false;
 					LastResult = false;
@@ -513,16 +519,16 @@ namespace DeploymentServer
 							try
 							{
 								TcpClient Client = Server.AcceptTcpClient();
-								Console.WriteLine("Got TCP connection.");
+								//LocalLog("Got TCP connection.");
 								NetworkStream ClStream = Client.GetStream();
 
 								TargetDevice = DeploymentProxy.Deployer.StartTCPTunnel(Device, ref TCPService);
 								if (TargetDevice == null)
 								{
-									Console.WriteLine("Cannot connect to device {0} for port forwarding.", Device);
+									LocalLog("Cannot connect to device " + Device + " for port forwarding.");
 									break;
 								}
-								Console.WriteLine("Connected to device.");
+								//LocalLog("Connected to device.");
 
 								Byte[] Buffer = new Byte[1024];
 
@@ -537,7 +543,7 @@ namespace DeploymentServer
 									{
 										if (Client.Client.Poll(10, SelectMode.SelectRead))
 										{
-											Console.WriteLine("TCP disconnected.");
+											//LocalLog("TCP disconnected.");
 											break;
 										}
 										System.Threading.Thread.Sleep(100);
@@ -550,7 +556,7 @@ namespace DeploymentServer
 							}
 							finally
 							{
-								Console.WriteLine("Port forwarding disconnected.");
+								//LocalLog("Port forwarding disconnected.");
 								if (TargetDevice != null)
 								{
 									TargetDevice.CloseTunnel(TCPService);
@@ -633,7 +639,7 @@ namespace DeploymentServer
 			}
 			catch (System.Exception ex)
 			{
-				Console.WriteLine("Failed to create deployment server process ({0})", ex.Message);
+				LocalLog("Failed to create deployment server process: " + ex.Message);
 			}
 		}
 
@@ -647,10 +653,10 @@ namespace DeploymentServer
 				{
 					if (LocalCommand == "stop")
 					{
-						Console.WriteLine("Deployment Server not running ...");
+						LocalLog("Deployment Server not running ...");
 						return null;
 					}
-					Console.WriteLine("Creating Deployment Server ...");
+					LocalLog("Creating Deployment Server ...");
 					CreateDeploymentServerProcess();
 				}
 				else
@@ -658,7 +664,7 @@ namespace DeploymentServer
 					if (RetryCount == 5)
 					{
 						ForceKillProcesses();
-						Console.WriteLine("Creating Deployment Server ...");
+						LocalLog("Creating Deployment Server ...");
 						CreateDeploymentServerProcess();
 					}
 				}
@@ -671,7 +677,7 @@ namespace DeploymentServer
 				}
 				catch (Exception e)
 				{
-					//Console.WriteLine("Retrying count {0} {1}", RetryCount.ToString(), e.Message);
+					//LocalLog("Retrying count {0} {1}", RetryCount.ToString(), e.Message);
 					RetryCount--;
 					if (Client != null)
 					{
@@ -732,7 +738,7 @@ namespace DeploymentServer
 										if (ArgTime > 0)
 										{
 											TimeOut = ArgTime;
-											Console.WriteLine(string.Format("Deployment Server timeout set to {0}", TimeOut.ToString()));
+											LocalLog(string.Format("Deployment Server timeout set to {0}", TimeOut.ToString()));
 										}
 									}
 									break;
@@ -775,7 +781,7 @@ namespace DeploymentServer
 			}
 			catch
 			{
-				Console.WriteLine("Exception creating Deployment Server Mutex. ");
+				LocalLog("Exception creating Deployment Server Mutex. ");
 				return;
 			}
 
@@ -809,10 +815,10 @@ namespace DeploymentServer
 												.Select(m => m.Value)
 												.ToList();
 				ParseServerParam(Arguments);
-				Console.WriteLine(string.Format("Deployment Server listening to port {0}", Port.ToString()));
-				Console.WriteLine(string.Format("Deployment Server inactivity timeout {0}", TimeOut.ToString()));
-				Console.WriteLine(string.Format("Deployment Server starting from {0}", TestStartPath));
-				Console.WriteLine("---------------------------------------------------------");
+				LocalLog(string.Format("Deployment Server listening to port {0}", Port.ToString()));
+				LocalLog(string.Format("Deployment Server inactivity timeout {0}", TimeOut.ToString()));
+				LocalLog(string.Format("Deployment Server starting from {0}", TestStartPath));
+				LocalLog("---------------------------------------------------------");
 
 				// Processing commands
 
@@ -852,12 +858,12 @@ namespace DeploymentServer
 					{
 						if (IsStopping)
 						{
-							Console.WriteLine("Deployment Server IsStopping exit.");
+							LocalLog("Deployment Server IsStopping exit.");
 							break;
 						}
 						if (GlobalTimer.ElapsedMilliseconds > TimeOut)
 						{
-							Console.WriteLine("Deployment Server inactivity timeout.");
+							LocalLog("Deployment Server inactivity timeout.");
 							IsStopping = true;
 							break;
 						}
@@ -866,13 +872,13 @@ namespace DeploymentServer
 			}
 			catch (SocketException e)
 			{
-				Console.WriteLine("SocketException: {0}", e);
+				LocalLog("SocketException: " + e.ToString());
 			}
 			catch (System.Exception Ex)
 			{
-				Console.WriteLine("Exception: {0}", Ex);
-				Console.WriteLine("Stack: {0}", Ex.StackTrace);
-				Console.WriteLine("Inner: {0}", Ex.InnerException.Message);
+				LocalLog("Exception: " + Ex.ToString());
+				LocalLog("Stack: " + Ex.StackTrace.ToString());
+				LocalLog("Inner: " + Ex.InnerException.Message.ToString());
 			}
 			finally
 			{
@@ -887,17 +893,17 @@ namespace DeploymentServer
 					catch (AbandonedMutexException)
 					{
 						//This catch is included to insure the program keeps running in the event this exception occurs.
-						Console.WriteLine("Deployment Server Mutex abandoned.");
+						LocalLog("Deployment Server Mutex abandoned.");
 					}
 					catch (ApplicationException)
 					{
 						//This catch is included to insure the program keeps running in the event this exception occurs.
-						Console.WriteLine("Deployment Server Mutex abandoned 1.");
+						LocalLog("Deployment Server Mutex abandoned 1.");
 					}
 					catch (SynchronizationLockException)
 					{
 						//This catch is included to insure the program keeps running in the event this exception occurs.
-						Console.WriteLine("Deployment Server Mutex abandoned. 2");
+						LocalLog("Deployment Server Mutex abandoned. 2");
 					}
 				}
 				if (ProcessClient != null)
@@ -908,7 +914,7 @@ namespace DeploymentServer
 				{
 					Server.Stop();
 				}
-				Console.WriteLine("Deployment Server Stopped.");
+				LocalLog("Deployment Server Stopped.");
 				Console.SetOut(OldConsole);
 				if (Writer != null)
 				{
@@ -944,7 +950,7 @@ namespace DeploymentServer
 				}
 				catch (System.Exception Ex)
 				{
-					Console.WriteLine(Ex.Message);
+					LocalLog(Ex.Message);
 				}
 			}
 			else
@@ -952,7 +958,7 @@ namespace DeploymentServer
 				if (LocalClientInfo.ParseCommand(Arguments))
 				{
 					LocalClientInfo.HasCommand = true;
-					Console.WriteLine("Running as local instance");
+					LocalLog("Running as local instance");
 					LocalClientInfo.RunCommand(Console.Out);
 					while (LocalClientInfo.IsStillRunning)
 					{
@@ -987,7 +993,7 @@ namespace DeploymentServer
 					{
 						if (!Response.Equals("[DSDIR]\"" + GetDeploymentServerPath() + "\"", StringComparison.InvariantCultureIgnoreCase))
 						{
-							Console.WriteLine("Wrong server running, restarting the server ...");
+							LocalLog("Wrong server running, restarting the server ...");
 							clientOut.Write("stop");
 							clientOut.Write("\r");
 							clientOut.Flush();
@@ -1000,7 +1006,7 @@ namespace DeploymentServer
 								}
 								else
 								{
-									Console.WriteLine(Response);
+									LocalLog(Response);
 								}
 							}
 							if (Client != null)
@@ -1044,7 +1050,7 @@ namespace DeploymentServer
 					}
 					else
 					{
-						Console.WriteLine(Response);
+						LocalLog(Response);
 					}
 					Thread.Sleep(10);
 				}
@@ -1053,7 +1059,7 @@ namespace DeploymentServer
 			{
 				if (LocalCommand != "stop")
 				{
-					Console.WriteLine("Exception: {0}", e);
+					LocalLog("Exception: " + e.ToString());
 				}
 			}
 			finally
@@ -1067,6 +1073,7 @@ namespace DeploymentServer
 
 		static int Main(string[] Args)
 		{
+			LocalConsole = Console.Out;
 			string LocalCommand = "";
 			if (Args.Length > 0)
 			{
@@ -1074,31 +1081,31 @@ namespace DeploymentServer
 			}
 			else
 			{
-				Console.WriteLine("Deployment Server usage: ");
-				Console.WriteLine("DeploymentServer.exe <command> [<parameter> [<value>] ...]");
-				Console.WriteLine("Valid Commands:");
-				Console.WriteLine("\t stop");
-				Console.WriteLine("\t backup");
-				Console.WriteLine("\t deploy");
-				Console.WriteLine("\t copyfile");
-				Console.WriteLine("\t install");
-				Console.WriteLine("\t enumerate");
-				Console.WriteLine("\t listdevices");
-				Console.WriteLine("\t listentodevice");
-				Console.WriteLine("\t command");
-				Console.WriteLine("\t forward");
-				Console.WriteLine("\t -iphonepackager");
-				Console.WriteLine("\t server");
-				Console.WriteLine("Valid Parameters:");
-				Console.WriteLine("\t -file <filename>");
-				Console.WriteLine("\t -bundle <bundle name>");
-				Console.WriteLine("\t -manifest <manifest file>");
-				Console.WriteLine("\t -ipa <ipa path>");
-				Console.WriteLine("\t -device <device ID>");
-				Console.WriteLine("\t -nokeepalive");
-				Console.WriteLine("\t -timeout <miliseconds>");
-				Console.WriteLine("\t -param <string parameter to be used for command>");
-				Console.WriteLine("");
+				LocalLog("Deployment Server usage: ");
+				LocalLog("DeploymentServer.exe <command> [<parameter> [<value>] ...]");
+				LocalLog("Valid Commands:");
+				LocalLog("\t stop");
+				LocalLog("\t backup");
+				LocalLog("\t deploy");
+				LocalLog("\t copyfile");
+				LocalLog("\t install");
+				LocalLog("\t enumerate");
+				LocalLog("\t listdevices");
+				LocalLog("\t listentodevice");
+				LocalLog("\t command");
+				LocalLog("\t forward");
+				LocalLog("\t -iphonepackager");
+				LocalLog("\t server");
+				LocalLog("Valid Parameters:");
+				LocalLog("\t -file <filename>");
+				LocalLog("\t -bundle <bundle name>");
+				LocalLog("\t -manifest <manifest file>");
+				LocalLog("\t -ipa <ipa path>");
+				LocalLog("\t -device <device ID>");
+				LocalLog("\t -nokeepalive");
+				LocalLog("\t -timeout <miliseconds>");
+				LocalLog("\t -param <string parameter to be used for command>");
+				LocalLog("");
 
 				return 0;
 			}
@@ -1167,7 +1174,7 @@ namespace DeploymentServer
 				{
 					NetworkStream ClStream = Client.GetStream();
 					string ClientIP = ((IPEndPoint)Client.Client.RemoteEndPoint).Address.ToString();
-					Console.WriteLine("Client [{0}] IP:{1} connected.", localID, ClientIP);
+					//LocalLog("Client [{0}] IP:{1} connected.", localID, ClientIP);
 
 					StreamWriter Writer = new StreamWriter(ClStream);
 					Writer.AutoFlush = true;
@@ -1182,7 +1189,7 @@ namespace DeploymentServer
 					
 					while (true)
 					{
-						//Console.WriteLine("Looping [{0}]", localID);
+						//LocalLog("Looping [{0}]", localID);
 						if (ClientInfo.HasCommand && !IsStopping)
 						{
 							if (!IsRunningCommand && !ClientInfo.IsStillRunning)
@@ -1247,7 +1254,7 @@ namespace DeploymentServer
 								// 10035 == WSAEWOULDBLOCK
 								if (e.NativeErrorCode.Equals(10035))
 								{
-									//Console.WriteLine("Still Connected, but the Send would block");
+									//LocalLog("Still Connected, but the Send would block");
 								}
 								else
 								{
@@ -1295,7 +1302,7 @@ namespace DeploymentServer
 				}
 				catch (Exception e)
 				{
-					Console.WriteLine("Exception: {0}", e);
+					LocalLog("Exception: " + e.ToString());
 				}
 				finally
 				{
@@ -1307,7 +1314,7 @@ namespace DeploymentServer
 					{
 						if (Client.Client != null && Client.Client.RemoteEndPoint != null)
 						{
-							Console.WriteLine("Client [{0}] disconnected ({1}).", localID, LastCommand);
+							//LocalLog("Client [{0}] disconnected ({1}).", localID, LastCommand);
 						}
 						Client.Close();
 					}
