#IPhonePackager
Add copyfileout command parameter to DeploymentServer project.  This is done to provide a method of retrieving a single specific file from an iOS device (log file, for example) without requiring a full backup of the bundle provided

diff --git a/Engine/Source/Programs/IOS/DeploymentServer/DeployTime.cs b/Engine/Source/Programs/IOS/DeploymentServer/DeployTime.cs
--- a/Engine/Source/Programs/IOS/DeploymentServer/DeployTime.cs
+++ b/Engine/Source/Programs/IOS/DeploymentServer/DeployTime.cs
@@ -544,7 +544,27 @@ namespace DeploymentServer
             });
         }
 
-        public bool BackupDocumentsDirectory(string BundleIdentifier, string DestinationDocumentsDirectory)
+		/// <summary>
+		/// Copies a single file from a specified location on the device for the given bundle
+		/// </summary>
+		public bool CopyFileFromDevice(string BundleIdentifier, string SourceFile, string DestFile)
+		{
+			return PerformActionOnAllDevices(StandardEnumerationDelayMS, delegate (MobileDeviceInstance Device)
+			{
+				bool bResult = true;
+				string DeviceType = Device.ProductType;
+				if (String.IsNullOrEmpty(DeviceId) || Device.DeviceId == DeviceId ||
+					(DeviceId.Contains("All_tvOS_On") && DeviceType.Contains("TV")) ||
+					(DeviceId.Contains("All_iOS_On") && !DeviceType.Contains("TV")))
+				{
+					bResult = Device.TryCopyFileOut(BundleIdentifier, SourceFile, DestFile);
+					ReportIF.Log("");
+				}
+				return bResult;
+			});
+		}
+
+		public bool BackupDocumentsDirectory(string BundleIdentifier, string DestinationDocumentsDirectory)
         {
             return PerformActionOnAllDevices(StandardEnumerationDelayMS, delegate(MobileDeviceInstance Device)
             {
diff --git a/Engine/Source/Programs/IOS/DeploymentServer/Program.cs b/Engine/Source/Programs/IOS/DeploymentServer/Program.cs
--- a/Engine/Source/Programs/IOS/DeploymentServer/Program.cs
+++ b/Engine/Source/Programs/IOS/DeploymentServer/Program.cs
@@ -311,6 +311,12 @@ namespace DeploymentServer
 								Writer.Flush();
 								break;
 
+							case "copyfileout":
+								Console.SetOut(Writer);
+								LastResult = DeploymentProxy.Deployer.CopyFileFromDevice(Bundle, FileList[0], FileList[1]);
+								Writer.Flush();
+								break;
+
 							case "install":
 								Console.SetOut(Writer);
 								LastResult = DeploymentProxy.Deployer.InstallIPAOnDevice(IpaPath);
@@ -1109,31 +1115,32 @@ namespace DeploymentServer
 			}
 			else
 			{
-				LocalLog("Deployment Server usage: ");
-				LocalLog("DeploymentServer.exe <command> [<parameter> [<value>] ...]");
-				LocalLog("Valid Commands:");
-				LocalLog("\t stop");
-				LocalLog("\t backup");
-				LocalLog("\t deploy");
-				LocalLog("\t copyfile");
-				LocalLog("\t install");
-				LocalLog("\t enumerate");
-				LocalLog("\t listdevices");
-				LocalLog("\t listentodevice");
-				LocalLog("\t command");
-				LocalLog("\t forward");
-				LocalLog("\t -iphonepackager");
-				LocalLog("\t server");
-				LocalLog("Valid Parameters:");
-				LocalLog("\t -file <filename>");
-				LocalLog("\t -bundle <bundle name>");
-				LocalLog("\t -manifest <manifest file>");
-				LocalLog("\t -ipa <ipa path>");
-				LocalLog("\t -device <device ID>");
-				LocalLog("\t -nokeepalive");
-				LocalLog("\t -timeout <miliseconds>");
-				LocalLog("\t -param <string parameter to be used for command>");
-				LocalLog("");
+				Console.WriteLine("Deployment Server usage: ");
+				Console.WriteLine("DeploymentServer.exe <command> [<parameter> [<value>] ...]");
+				Console.WriteLine("Valid Commands:");
+				Console.WriteLine("\t stop");
+				Console.WriteLine("\t backup");
+				Console.WriteLine("\t deploy");
+				Console.WriteLine("\t copyfile");
+				Console.WriteLine("\t copyfileout");
+				Console.WriteLine("\t install");
+				Console.WriteLine("\t enumerate");
+				Console.WriteLine("\t listdevices");
+				Console.WriteLine("\t listentodevice");
+				Console.WriteLine("\t command");
+				Console.WriteLine("\t forward");
+				Console.WriteLine("\t -iphonepackager");
+				Console.WriteLine("\t server");
+				Console.WriteLine("Valid Parameters:");
+				Console.WriteLine("\t -file <filename>");
+				Console.WriteLine("\t -bundle <bundle name>");
+				Console.WriteLine("\t -manifest <manifest file>");
+				Console.WriteLine("\t -ipa <ipa path>");
+				Console.WriteLine("\t -device <device ID>");
+				Console.WriteLine("\t -nokeepalive");
+				Console.WriteLine("\t -timeout <miliseconds>");
+				Console.WriteLine("\t -param <string parameter to be used for command>");
+				Console.WriteLine("");
 
 				return 0;
 			}
diff --git a/Engine/Source/Programs/IOS/MobileDeviceInterface/MobileDeviceInstance.cs b/Engine/Source/Programs/IOS/MobileDeviceInterface/MobileDeviceInstance.cs
--- a/Engine/Source/Programs/IOS/MobileDeviceInterface/MobileDeviceInstance.cs
+++ b/Engine/Source/Programs/IOS/MobileDeviceInterface/MobileDeviceInstance.cs
@@ -596,11 +596,38 @@ namespace Manzana
 			}
         }
 
-        /// <summary>
-        /// Tries to copy all of the files in a particular directory on the PC to the phone directory
-        /// (requires the bundle identifier to be able to mount that directory)
-        /// </summary>
-        public bool TryCopy(string BundleIdentifier, string SourceFolderOnPC, string TargetFolderOnDevice)
+		/// <summary>
+		/// Tries to copy the specified file from the device to a specified file path on the PC.
+		/// (requires the bundle identifier to be able to mount that directory)
+		/// </summary>
+		public bool TryCopyFileOut(string BundleIdentifier, string SourceFileOnDevice, string TargetFileOnPC)
+		{
+			if (ConnectToBundle(BundleIdentifier))
+			{
+				WriteProgressLine("Connected to bundle '{0}'", 0, BundleIdentifier);
+				try
+				{
+					CopyFileFromPhone(TargetFileOnPC, SourceFileOnDevice, 1024 * 1024);
+					return true;
+				}
+				catch (Exception ex)
+				{
+					WriteProgressLine("Failed to transfer a file, extended error is '{0}'", 100, ex.Message);
+					return false;
+				}
+			}
+			else
+			{
+				WriteProgressLine("Error: Failed to connect to bundle '{0}'", 100, BundleIdentifier);
+				return false;
+			}
+		}
+
+		/// <summary>
+		/// Tries to copy all of the files in a particular directory on the PC to the phone directory
+		/// (requires the bundle identifier to be able to mount that directory)
+		/// </summary>
+		public bool TryCopy(string BundleIdentifier, string SourceFolderOnPC, string TargetFolderOnDevice)
 		{
 			if (ConnectToBundle(BundleIdentifier))
 			{
@@ -1152,10 +1179,10 @@ namespace Manzana
             }
         }
 
-        /// <summary>
-        /// Copies a file from the phone to the PC
-        /// </summary>
-        void CopyFileFromPhone(string PathOnPC, string PathOnPhone, int ChunkSize)
+		/// <summary>
+		/// Copies a file from the phone to the PC
+		/// </summary>
+		void CopyFileFromPhone(string PathOnPC, string PathOnPhone, int ChunkSize)
         {
             DateTime StartTime = DateTime.Now;
 
