UE-64937 - IphonePackager cannot verify UDID on new devices (XS or newer)

diff --git a/Engine/Source/Programs/IOS/iPhonePackager/GraphicalResignTool.cs b/Engine/Source/Programs/IOS/iPhonePackager/GraphicalResignTool.cs
--- a/Engine/Source/Programs/IOS/iPhonePackager/GraphicalResignTool.cs
+++ b/Engine/Source/Programs/IOS/iPhonePackager/GraphicalResignTool.cs
@@ -1,4 +1,4 @@
-﻿/**
+/**
  * Copyright 1998-2018 Epic Games, Inc. All Rights Reserved.
  */
 
@@ -566,9 +566,9 @@ namespace iPhonePackager
 		{
 			string UDID = ScanningUDIDEdit.Text.Trim();
 
-			if (UDID.Length != 40)
+			if (UDID.Length <= 0)
 			{
-				ShowErrorMessage("The UDID must be exactly 40 characters long");
+				ShowErrorMessage("The UDID value is not properly set");
 			}
 			else
 			{
@@ -580,9 +580,9 @@ namespace iPhonePackager
 		{
 			string UDID = ScanningUDIDEdit.Text.Trim();
 
-			if (UDID.Length != 40)
+			if (UDID.Length <= 0)
 			{
-				ShowErrorMessage("The UDID must be exactly 40 characters long");
+				ShowErrorMessage("The UDID value is not properly set");
 			}
 			else
 			{
