Support for Windows Store version of iTunes.

diff --git a/Engine/Source/Developer/IOS/IOSTargetPlatform/Private/IOSTargetPlatform.cpp b/Engine/Source/Developer/IOS/IOSTargetPlatform/Private/IOSTargetPlatform.cpp
--- a/Engine/Source/Developer/IOS/IOSTargetPlatform/Private/IOSTargetPlatform.cpp
+++ b/Engine/Source/Developer/IOS/IOSTargetPlatform/Private/IOSTargetPlatform.cpp
@@ -147,6 +147,54 @@ bool FIOSTargetPlatform::IsSdkInstalled(bool bProjectHasCode, FString& OutTutori
 		biOSSDKInstalled = true;
 	}
 	
+	// Check for iTunes 12, Windows Store version
+	if (!biOSSDKInstalled)
+	{
+		HKEY hPackagesKey;
+		FString PackagesKeyName(TEXT("Software\\Classes\\Local Settings\\Software\\Microsoft\\Windows\\CurrentVersion\\AppModel\\PackageRepository\\Packages"));
+		if (RegOpenKeyEx(HKEY_LOCAL_MACHINE, *PackagesKeyName, 0, KEY_READ, &hPackagesKey) == ERROR_SUCCESS)
+		{
+			DWORD NumSubKeys = 0;
+			DWORD LongestSubKeyLength;
+			if (RegQueryInfoKey(hPackagesKey, NULL, NULL, NULL, &NumSubKeys, &LongestSubKeyLength, NULL, NULL, NULL, NULL, NULL, NULL) == ERROR_SUCCESS)
+			{
+				static const uint32 MAX_SUBKEY_LENGTH = 512;
+				check(LongestSubKeyLength < MAX_SUBKEY_LENGTH);
+
+				for (uint32 i = 0; i < NumSubKeys; i++)
+				{
+					TCHAR SubKeyName[MAX_SUBKEY_LENGTH];
+					DWORD SubKeyLength = MAX_SUBKEY_LENGTH;
+					if (RegEnumKeyEx(hPackagesKey, i, SubKeyName, &SubKeyLength, NULL, NULL, NULL, NULL) == ERROR_SUCCESS)
+					{
+						check(SubKeyLength < MAX_SUBKEY_LENGTH);
+
+						if (FString(SubKeyName).Contains(TEXT("AppleInc.iTunes")) && (FString(SubKeyName).Contains(TEXT("_x64")) || FString(SubKeyName).Contains(TEXT("_x86"))))
+						{
+							HKEY iTunesKey;
+							TCHAR FullPackageSubKeyName[MAX_SUBKEY_LENGTH + MAX_SUBKEY_LENGTH + 1];
+							_stprintf_s(FullPackageSubKeyName, MAX_SUBKEY_LENGTH + MAX_SUBKEY_LENGTH + 1, TEXT("%s\\%s"), *PackagesKeyName, SubKeyName);
+							if (RegOpenKeyEx(HKEY_LOCAL_MACHINE, FullPackageSubKeyName, 0, KEY_READ, &iTunesKey) == ERROR_SUCCESS)
+							{
+								if (RegQueryValueEx(iTunesKey, TEXT("Path"), NULL, NULL, (BYTE*)dllPath, &pathSize) == ERROR_SUCCESS)
+								{
+									TCHAR dllFullPath[MAX_SUBKEY_LENGTH + MAX_SUBKEY_LENGTH + 32];
+									_stprintf_s(dllFullPath, MAX_SUBKEY_LENGTH + MAX_SUBKEY_LENGTH + 32, TEXT("%s\\AMDS32\\MobileDevice.dll"), dllPath);
+
+									if (IFileManager::Get().FileSize(*FString(dllFullPath)) != INDEX_NONE)
+									{
+										biOSSDKInstalled = true;
+										break;
+									}
+								}
+							}
+						}
+					}
+				}
+			}
+		}
+	}
+
 	// Check for iTunes 11
 	if(!biOSSDKInstalled
 		&& RegOpenKeyEx(HKEY_LOCAL_MACHINE, TEXT("SOFTWARE\\Wow6432Node\\Apple Inc.\\Apple Mobile Device Support\\Shared"), 0, KEY_READ, &hKey) == ERROR_SUCCESS
diff --git a/Engine/Source/Programs/IOS/MobileDeviceInterface/MobileDevice.cs b/Engine/Source/Programs/IOS/MobileDeviceInterface/MobileDevice.cs
--- a/Engine/Source/Programs/IOS/MobileDeviceInterface/MobileDevice.cs
+++ b/Engine/Source/Programs/IOS/MobileDeviceInterface/MobileDevice.cs
@@ -181,6 +181,56 @@ namespace Manzana
 		public static IntPtr kCFRunLoopDefaultMode = CFStringCreateWithCString(IntPtr.Zero, "kCFRunLoopDefaultMode", 0);
 	}
 	
+	public static class ITunesFinder
+	{
+		public static string FindWindowsStoreITunesInstallPath()
+		{
+			string InstallPath = null;
+			string PackagesKeyName = "Software\\Classes\\Local Settings\\Software\\Microsoft\\Windows\\CurrentVersion\\AppModel\\PackageRepository\\Packages";
+
+			RegistryKey PackagesKey = Registry.LocalMachine.OpenSubKey(PackagesKeyName);
+
+			if (PackagesKey != null)
+			{
+				string[] PackageSubKeyNames = PackagesKey.GetSubKeyNames();
+
+				foreach (string PackageSubKeyName in PackageSubKeyNames)
+				{
+					if (PackageSubKeyName.Contains("AppleInc.iTunes") && (PackageSubKeyName.Contains("_x64") || PackageSubKeyName.Contains("_x86")))
+					{
+						string FullPackageSubKeyName = PackagesKeyName + "\\" + PackageSubKeyName;
+
+						RegistryKey iTunesKey = Registry.LocalMachine.OpenSubKey(FullPackageSubKeyName);
+						if (iTunesKey != null)
+						{
+							InstallPath = (string)iTunesKey.GetValue("Path");
+							break;
+						}
+					}
+				}
+			}
+
+			return InstallPath;
+		}
+
+		public static string FindWindowsStoreITunesDLL()
+		{
+			string DllPath = null;
+
+			string InstallPath = FindWindowsStoreITunesInstallPath();
+
+			if (InstallPath != null)
+			{
+				DllPath = InstallPath + "\\AMDS32\\MobileDevice.dll";
+
+				// There's also a 64-bit version here, but it seems that our c# programs are 32-bit
+				// DllPath = InstallPath + "\\AMDS64\\MobileDevice.dll";
+			}
+
+			return DllPath;
+		}
+	}
+
 	public class CoreFoundationRunLoop
 	{
 		static public ICoreFoundationRunLoop CoreImpl;
@@ -208,6 +258,18 @@ namespace Manzana
 					}
 				}
 
+				string ITunesInstallPath = ITunesFinder.FindWindowsStoreITunesInstallPath();
+				if(ITunesInstallPath != null)
+				{
+					DirectoryInfo iTunesInstallDirectory = new DirectoryInfo(ITunesInstallPath);
+
+					if (iTunesInstallDirectory.Exists)
+					{
+						PathBits.Add(iTunesInstallDirectory.FullName);
+						PathBits.Add(ITunesFinder.FindWindowsStoreITunesInstallPath() + "\\VFS\\ProgramFilesCommonX86\\Apple\\Apple Application Support");
+					}
+				}
+
 				// Add some guesses as well
 				DirectoryInfo AppleMobileDeviceSupport = new DirectoryInfo(Environment.GetFolderPath(Environment.SpecialFolder.CommonProgramFiles) + @"\Apple\Mobile Device Support");
 				if (AppleMobileDeviceSupport.Exists)
@@ -1871,6 +1933,18 @@ namespace Manzana
 				}
 			}
 
+			string iTunesWindowsStoreRegistryDllPath = ITunesFinder.FindWindowsStoreITunesDLL();
+			if(iTunesWindowsStoreRegistryDllPath != null)
+			{
+				FileInfo iTunesMobileDeviceFile = new FileInfo(iTunesWindowsStoreRegistryDllPath);
+
+				if (iTunesMobileDeviceFile.Exists)
+				{
+					PathBits.Add(iTunesMobileDeviceFile.DirectoryName);
+					PathBits.Add(ITunesFinder.FindWindowsStoreITunesInstallPath() + "\\VFS\\ProgramFilesCommonX86\\Apple\\Apple Application Support");
+				}
+			}
+
 			object RegistrySupportDir = Registry.GetValue(@"HKEY_LOCAL_MACHINE\SOFTWARE\Apple Inc.\Apple Application Support", "InstallDir", Environment.CurrentDirectory);
 			if (RegistrySupportDir != null)
 			{
@@ -2647,7 +2721,7 @@ namespace Manzana
     {
 		static public MobileDeviceImpl DeviceImpl = null;
 
-        static MobileDevice()
+		static MobileDevice()
         {
 			if (Environment.OSVersion.Platform == PlatformID.MacOSX || Environment.OSVersion.Platform == PlatformID.Unix)
 			{
@@ -2657,7 +2731,9 @@ namespace Manzana
 			{
 				string dllPath11 = Microsoft.Win32.Registry.GetValue("HKEY_LOCAL_MACHINE\\SOFTWARE\\Wow6432Node\\Apple Inc.\\Apple Mobile Device Support\\Shared", "iTunesMobileDeviceDLL", null) as string;
 				string dllPath12 = Microsoft.Win32.Registry.GetValue("HKEY_LOCAL_MACHINE\\SOFTWARE\\Wow6432Node\\Apple Inc.\\Apple Mobile Device Support\\Shared", "MobileDeviceDLL", null) as string;
-				if (!String.IsNullOrEmpty(dllPath12) && File.Exists(dllPath12))
+				string dllPath12WindowsStore = ITunesFinder.FindWindowsStoreITunesDLL();
+
+				if ((!String.IsNullOrEmpty(dllPath12) && File.Exists(dllPath12)) || (!String.IsNullOrEmpty(dllPath12WindowsStore) && File.Exists(dllPath12WindowsStore)))
 				{
 					DeviceImpl = new MobileDeviceWiniTunes12();
 				}
diff --git a/Engine/Source/Programs/IOS/iPhonePackager/iPhonePackager.cs b/Engine/Source/Programs/IOS/iPhonePackager/iPhonePackager.cs
--- a/Engine/Source/Programs/IOS/iPhonePackager/iPhonePackager.cs
+++ b/Engine/Source/Programs/IOS/iPhonePackager/iPhonePackager.cs
@@ -15,6 +15,7 @@ using System.Windows.Forms;
 using Ionic.Zlib;
 using System.ComponentModel;
 using System.Security.Cryptography.X509Certificates;
+using Microsoft.Win32;
 
 namespace iPhonePackager
 {
@@ -705,6 +706,36 @@ namespace iPhonePackager
 			}
 		}
 
+		static string FindWindowsStoreITunesDLL()
+		{
+			string InstallPath = null;
+
+			string PackagesKeyName = "Software\\Classes\\Local Settings\\Software\\Microsoft\\Windows\\CurrentVersion\\AppModel\\PackageRepository\\Packages";
+			
+			RegistryKey PackagesKey = Registry.LocalMachine.OpenSubKey(PackagesKeyName);
+			if (PackagesKey != null)
+			{
+				string[] PackageSubKeyNames = PackagesKey.GetSubKeyNames();
+				
+				foreach (string PackageSubKeyName in PackageSubKeyNames)
+				{
+					if (PackageSubKeyName.Contains("AppleInc.iTunes") && (PackageSubKeyName.Contains("_x64") || PackageSubKeyName.Contains("_x86")))
+					{
+						string FullPackageSubKeyName = PackagesKeyName + "\\" + PackageSubKeyName;
+
+						RegistryKey iTunesKey = Registry.LocalMachine.OpenSubKey(FullPackageSubKeyName);
+						if (iTunesKey != null)
+						{
+							InstallPath = (string)iTunesKey.GetValue("Path") + "\\AMDS32\\MobileDevice.dll";
+							break;
+						}
+					}
+				}
+			}
+
+			return InstallPath;
+		}
+
 		/**
 		 * Main control loop
 		 */
@@ -806,10 +837,14 @@ namespace iPhonePackager
                         }
                         else
 						{
-							dllPath = Microsoft.Win32.Registry.GetValue("HKEY_LOCAL_MACHINE\\SOFTWARE\\Wow6432Node\\Apple Inc.\\Apple Mobile Device Support\\Shared", "iTunesMobileDeviceDLL", null) as string;
+							dllPath = Registry.GetValue("HKEY_LOCAL_MACHINE\\SOFTWARE\\Wow6432Node\\Apple Inc.\\Apple Mobile Device Support\\Shared", "iTunesMobileDeviceDLL", null) as string;
 							if (String.IsNullOrEmpty(dllPath) || !File.Exists(dllPath))
 							{
-								dllPath = Microsoft.Win32.Registry.GetValue("HKEY_LOCAL_MACHINE\\SOFTWARE\\Wow6432Node\\Apple Inc.\\Apple Mobile Device Support\\Shared", "MobileDeviceDLL", null) as string;
+								dllPath = Registry.GetValue("HKEY_LOCAL_MACHINE\\SOFTWARE\\Wow6432Node\\Apple Inc.\\Apple Mobile Device Support\\Shared", "MobileDeviceDLL", null) as string;
+								if (String.IsNullOrEmpty(dllPath) || !File.Exists(dllPath))
+								{
+									dllPath = FindWindowsStoreITunesDLL();
+								}
 							}
 						}
                         if (String.IsNullOrEmpty(dllPath) || (!File.Exists(dllPath) && !Directory.Exists(dllPath)))
diff --git a/Engine/Source/Programs/UnrealBuildTool/Platform/IOS/UEBuildIOS.cs b/Engine/Source/Programs/UnrealBuildTool/Platform/IOS/UEBuildIOS.cs
--- a/Engine/Source/Programs/UnrealBuildTool/Platform/IOS/UEBuildIOS.cs
+++ b/Engine/Source/Programs/UnrealBuildTool/Platform/IOS/UEBuildIOS.cs
@@ -8,6 +8,7 @@ using System.IO;
 using System.Linq;
 using System.Xml;
 using Tools.DotNETCommon;
+using Microsoft.Win32;
 
 namespace UnrealBuildTool
 {
@@ -993,18 +994,53 @@ namespace UnrealBuildTool
 			if (!Utils.IsRunningOnMono)
 			{
 				// check to see if iTunes is installed
-				string dllPath = Microsoft.Win32.Registry.GetValue("HKEY_LOCAL_MACHINE\\SOFTWARE\\Wow6432Node\\Apple Inc.\\Apple Mobile Device Support\\Shared", "iTunesMobileDeviceDLL", null) as string;
+				string dllPath = Registry.GetValue("HKEY_LOCAL_MACHINE\\SOFTWARE\\Wow6432Node\\Apple Inc.\\Apple Mobile Device Support\\Shared", "iTunesMobileDeviceDLL", null) as string;
 				if (String.IsNullOrEmpty(dllPath) || !File.Exists(dllPath))
 				{
-					dllPath = Microsoft.Win32.Registry.GetValue("HKEY_LOCAL_MACHINE\\SOFTWARE\\Wow6432Node\\Apple Inc.\\Apple Mobile Device Support\\Shared", "MobileDeviceDLL", null) as string;
+					dllPath = Registry.GetValue("HKEY_LOCAL_MACHINE\\SOFTWARE\\Wow6432Node\\Apple Inc.\\Apple Mobile Device Support\\Shared", "MobileDeviceDLL", null) as string;
 					if (String.IsNullOrEmpty(dllPath) || !File.Exists(dllPath))
 					{
-						return SDKStatus.Invalid;
+						dllPath = FindWindowsStoreITunesDLL();
+
+						if (String.IsNullOrEmpty(dllPath) || !File.Exists(dllPath))
+						{
+							return SDKStatus.Invalid;
+						}
 					}
 				}
 			}
 			return SDKStatus.Valid;
 		}
+
+		static string FindWindowsStoreITunesDLL()
+		{
+			string InstallPath = null;
+
+			string PackagesKeyName = "Software\\Classes\\Local Settings\\Software\\Microsoft\\Windows\\CurrentVersion\\AppModel\\PackageRepository\\Packages";
+
+			RegistryKey PackagesKey = Registry.LocalMachine.OpenSubKey(PackagesKeyName);
+			if (PackagesKey != null)
+			{
+				string[] PackageSubKeyNames = PackagesKey.GetSubKeyNames();
+
+				foreach (string PackageSubKeyName in PackageSubKeyNames)
+				{
+					if (PackageSubKeyName.Contains("AppleInc.iTunes") && (PackageSubKeyName.Contains("_x64") || PackageSubKeyName.Contains("_x86")))
+					{
+						string FullPackageSubKeyName = PackagesKeyName + "\\" + PackageSubKeyName;
+
+						RegistryKey iTunesKey = Registry.LocalMachine.OpenSubKey(FullPackageSubKeyName);
+						if (iTunesKey != null)
+						{
+							InstallPath = (string)iTunesKey.GetValue("Path") + "\\AMDS32\\MobileDevice.dll";
+							break;
+						}
+					}
+				}
+			}
+
+			return InstallPath;
+		}
 	}
 
 	class IOSPlatformFactory : UEBuildPlatformFactory
