Support for iOS Client builds.

diff --git a/Engine/Source/Developer/LauncherServices/Private/Launcher/LauncherWorker.cpp b/Engine/Source/Developer/LauncherServices/Private/Launcher/LauncherWorker.cpp
--- a/Engine/Source/Developer/LauncherServices/Private/Launcher/LauncherWorker.cpp
+++ b/Engine/Source/Developer/LauncherServices/Private/Launcher/LauncherWorker.cpp
@@ -336,6 +336,10 @@ FString FLauncherWorker::CreateUATCommand( const ILauncherProfileRef& InProfile,
 			{
 				Platforms += TEXT("+Mac");
 			}
+			else if (PlatformInfo->TargetPlatformName == FName("IOSClient"))
+			{
+				Platforms += TEXT("+IOS");
+			}
 			else
 			{
 				Platforms += TEXT("+");
diff --git a/Engine/Source/Programs/AutomationTool/IOS/IOSPlatform.Automation.cs b/Engine/Source/Programs/AutomationTool/IOS/IOSPlatform.Automation.cs
--- a/Engine/Source/Programs/AutomationTool/IOS/IOSPlatform.Automation.cs
+++ b/Engine/Source/Programs/AutomationTool/IOS/IOSPlatform.Automation.cs
@@ -448,7 +448,9 @@ public class IOSPlatform : Platform
 				}
 
 				IPPArguments += (cookonthefly ? " -cookonthefly" : "");
-				IPPArguments += " -stagedir \"" + CombinePaths(Params.BaseStageDirectory, PlatformName) + "\"";
+
+				string CookPlatformName = GetCookPlatform(Params.DedicatedServer, Params.Client);
+				IPPArguments += " -stagedir \"" + CombinePaths(Params.BaseStageDirectory, CookPlatformName) + "\"";
 				IPPArguments += " -project \"" + Params.RawProjectPath + "\"";
 				if (Params.IterativeDeploy)
 				{
diff --git a/Engine/Source/Programs/IOS/iPhonePackager/CompileTime.cs b/Engine/Source/Programs/IOS/iPhonePackager/CompileTime.cs
--- a/Engine/Source/Programs/IOS/iPhonePackager/CompileTime.cs
+++ b/Engine/Source/Programs/IOS/iPhonePackager/CompileTime.cs
@@ -54,14 +54,14 @@ namespace iPhonePackager
 		/** /MacStagingRootDir/Payload/GameName.app */
 		protected static string RemoteAppDirectory
 		{
-			get { return RemoteAppPayloadDirectory + "/" + Program.GameName + Program.Architecture + ".app"; }
+			get { return RemoteAppPayloadDirectory + "/" + Program.GameName + (Program.IsClient ? "Client" : "") + Program.Architecture + ".app"; }
 		}
 
 
 		/** /MacStagingRootDir/Payload/GameName.app/GameName */
 		protected static string RemoteExecutablePath
 		{
-			get { return RemoteAppDirectory + "/" + Program.GameName + Program.Architecture; }
+			get { return RemoteAppDirectory + "/" + Program.GameName + (Program.IsClient ? "Client" : "") + Program.Architecture; }
 		}
 
 		private static string CurrentBaseXCodeCommandLine;
@@ -524,7 +524,7 @@ namespace iPhonePackager
 
 				// NOTE: -y preserves symbolic links which is needed for iOS distro builds
 				// -x excludes a file (excluding the dSYM keeps sizes smaller, and it shouldn't be in the IPA anyways)
-				string dSYMName = "Payload/" + Program.GameName + Program.Architecture + ".app.dSYM";
+				string dSYMName = "Payload/" + Program.GameName + (Program.IsClient ? "Client" : "") + Program.Architecture + ".app.dSYM";
 				DisplayCommandLine = String.Format("zip -q -r -y -{0} -T {1} Payload iTunesArtwork -x {2}/ -x {2}/* " +
 					"-x {2}/Contents/ -x {2}/Contents/* -x {2}/Contents/Resources/ -x {2}/Contents/Resources/* " +
 					" -x {2}/Contents/Resources/DWARF/ -x {2}/Contents/Resources/DWARF/*",
@@ -539,8 +539,8 @@ namespace iPhonePackager
 			case "gendsym":
 				Program.Log( " ... generating DSYM" );
 
-				string ExePath  = "Payload/" + Program.GameName + ".app/" + Program.GameName;
-				string dSYMPath = Program.GameName + ".app.dSYM";
+				string ExePath  = "Payload/" + Program.GameName + (Program.IsClient ? "Client" : "") + ".app/" + Program.GameName + (Program.IsClient ? "Client" : "");
+				string dSYMPath = Program.GameName + (Program.IsClient ? "Client" : "") + ".app.dSYM";
 				DisplayCommandLine = String.Format("dsymutil -o {0} {1}", dSYMPath, ExePath);
 
 				CommandLine = "\"" + MacStagingRootDir + "\"" + DisplayCommandLine;
diff --git a/Engine/Source/Programs/IOS/iPhonePackager/Config.cs b/Engine/Source/Programs/IOS/iPhonePackager/Config.cs
--- a/Engine/Source/Programs/IOS/iPhonePackager/Config.cs
+++ b/Engine/Source/Programs/IOS/iPhonePackager/Config.cs
@@ -168,7 +168,7 @@ namespace iPhonePackager
 		//@TODO: Deprecate this directory
 		public static string PayloadDirectory
 		{
-			get { return Path.GetFullPath(PCStagingRootDir + @"\Payload\" + Program.GameName + Program.Architecture + ".app"); }
+			get { return Path.GetFullPath(PCStagingRootDir + @"\Payload\" + Program.GameName + (Program.IsClient ? "Client" : "") + Program.Architecture + ".app"); }
 		}
 
 		/// <summary>
@@ -196,11 +196,11 @@ namespace iPhonePackager
 			{
 				if (Program.GameConfiguration == "Development")
 				{
-					return SigningPrefix + Program.GameName + Program.Architecture + ".ipa";
+					return SigningPrefix + Program.GameName + (Program.IsClient ? "Client" : "") + Program.Architecture + ".ipa";
 				}
 				else
 				{
-					return SigningPrefix + Program.GameName + "-" + Config.OSString + "-" + Program.GameConfiguration + Program.Architecture + ".ipa";
+					return SigningPrefix + Program.GameName + (Program.IsClient ? "Client" : "") + "-" + Config.OSString + "-" + Program.GameConfiguration + Program.Architecture + ".ipa";
 				}
 			}
 		}
@@ -236,11 +236,11 @@ namespace iPhonePackager
 
 			if (Program.GameConfiguration == "Development")
 			{
-				Filename = Path.Combine(Config.BinariesDirectory, FilePrefix + Program.GameName + Program.Architecture + FileSuffix);
+				Filename = Path.Combine(Config.BinariesDirectory, FilePrefix + Program.GameName + (Program.IsClient ? "Client" : "") + Program.Architecture + FileSuffix);
 			}
 			else
 			{
-				Filename = Path.Combine(Config.BinariesDirectory, FilePrefix + Program.GameName + "-" + Config.OSString + "-" + Program.GameConfiguration + Program.Architecture + FileSuffix);
+				Filename = Path.Combine(Config.BinariesDirectory, FilePrefix + Program.GameName + (Program.IsClient ? "Client" : "") + "-" + Config.OSString + "-" + Program.GameConfiguration + Program.Architecture + FileSuffix);
 			}
 
 			return Filename;
@@ -262,11 +262,11 @@ namespace iPhonePackager
 			}
 			if (Program.GameConfiguration == "Development")
 			{
-				Filename = Path.Combine(BinariesDir, FilePrefix + (bIsCodeBasedProject ? GameName : "UE4Game") + Program.Architecture + FileSuffix);
+				Filename = Path.Combine(BinariesDir, FilePrefix + (bIsCodeBasedProject ? GameName : "UE4Game") + (Program.IsClient ? "Client" : "") + Program.Architecture + FileSuffix);
 			}
 			else
 			{
-				Filename = Path.Combine(BinariesDir, FilePrefix + (bIsCodeBasedProject ? GameName : "UE4Game") + "-" + Config.OSString + "-" + Program.GameConfiguration + Program.Architecture + FileSuffix);
+				Filename = Path.Combine(BinariesDir, FilePrefix + (bIsCodeBasedProject ? GameName : "UE4Game") + (Program.IsClient ? "Client" : "") + "-" + Config.OSString + "-" + Program.GameConfiguration + Program.Architecture + FileSuffix);
 			}
 
 			// ensure the directory exists
@@ -281,11 +281,11 @@ namespace iPhonePackager
 		{
 			if (Program.GameConfiguration == "Development")
 			{
-				return FileOperations.FindPrefixedFile(Config.BinariesDirectory, Program.GameName + Program.Architecture + FileSuffix);
+				return FileOperations.FindPrefixedFile(Config.BinariesDirectory, Program.GameName + (Program.IsClient ? "Client" : "") + Program.Architecture + FileSuffix);
 			}
 			else
 			{
-				return FileOperations.FindPrefixedFile(Config.BinariesDirectory, Program.GameName + "-" + Config.OSString + "-" + Program.GameConfiguration + Program.Architecture + FileSuffix);
+				return FileOperations.FindPrefixedFile(Config.BinariesDirectory, Program.GameName + (Program.IsClient ? "Client" : "") + "-" + Config.OSString + "-" + Program.GameConfiguration + Program.Architecture + FileSuffix);
 			}
 		}
 
@@ -396,7 +396,7 @@ namespace iPhonePackager
 		/// </summary>
 		public static string AppDirectoryInZIP
 		{
-			get { return String.Format("Payload/{0}{1}.app", Program.GameName, Program.Architecture); }
+			get { return String.Format("Payload/{0}{1}.app", Program.GameName + (Program.IsClient ? "Client" : ""), Program.Architecture); }
 		}
 
 		/// <summary>
@@ -484,7 +484,7 @@ namespace iPhonePackager
 			// Root directory on PC for staging files to copy to Mac
 			Config.PCStagingRootDir = String.Format(@"{0}-Deploy\{1}\{2}{3}\",
 				IntermediateDirectory,
-				Program.GameName,
+				Program.GameName + (Program.IsClient ? "Client" : ""),
 				Program.GameConfiguration,
 				Program.Architecture);
 
diff --git a/Engine/Source/Programs/IOS/iPhonePackager/CookTime.cs b/Engine/Source/Programs/IOS/iPhonePackager/CookTime.cs
--- a/Engine/Source/Programs/IOS/iPhonePackager/CookTime.cs
+++ b/Engine/Source/Programs/IOS/iPhonePackager/CookTime.cs
@@ -151,11 +151,11 @@ namespace iPhonePackager
 			string NameDecoration;
 			if (Program.GameConfiguration == "Development")
 			{
-				NameDecoration = Program.Architecture;
+				NameDecoration = (Program.IsClient ? "Client" : "") + Program.Architecture;
 			}
 			else
 			{
-				NameDecoration = "-" + Config.OSString + "-" + Program.GameConfiguration + Program.Architecture;
+				NameDecoration = (Program.IsClient ? "Client" : "") + "-" + Config.OSString + "-" + Program.GameConfiguration + Program.Architecture;
 			}
 
 			// Copy and un-decorate the binary name
@@ -221,7 +221,7 @@ namespace iPhonePackager
 				return;
 			}
 
-			string ZipWorkingDir = String.Format("Payload/{0}{1}.app/", Program.GameName, Program.Architecture);
+			string ZipWorkingDir = String.Format("Payload/{0}{1}.app/", Program.GameName + (Program.IsClient ? "Client" : ""), Program.Architecture);
 
 			FileOperations.ZipFileSystem FileSystem = new FileOperations.ZipFileSystem(Zip, ZipWorkingDir);
 
@@ -313,7 +313,7 @@ namespace iPhonePackager
 					string RelativeFilename = AbsoluteFilename.Substring(SourceDir.Length + 1).Replace('\\', '/');
 
 					string ZipAbsolutePath = String.Format("Payload/{0}{1}.app/{2}",
-						Program.GameName,
+						Program.GameName + (Program.IsClient ? "Client" : ""),
 						Program.Architecture,
 						RelativeFilename);
 
diff --git a/Engine/Source/Programs/IOS/iPhonePackager/Utilities.cs b/Engine/Source/Programs/IOS/iPhonePackager/Utilities.cs
--- a/Engine/Source/Programs/IOS/iPhonePackager/Utilities.cs
+++ b/Engine/Source/Programs/IOS/iPhonePackager/Utilities.cs
@@ -574,7 +574,7 @@ namespace iPhonePackager
 
 		static public string GetPrecompileSourcePListFilename()
 		{
-			// check for one in he project directory
+			// check for one in the project directory
 			string SourceName = FileOperations.FindPrefixedFile(Config.IntermediateDirectory, Program.GameName + "-Info.plist");
 
 			if (!File.Exists(SourceName))
diff --git a/Engine/Source/Programs/IOS/iPhonePackager/iPhonePackager.cs b/Engine/Source/Programs/IOS/iPhonePackager/iPhonePackager.cs
--- a/Engine/Source/Programs/IOS/iPhonePackager/iPhonePackager.cs
+++ b/Engine/Source/Programs/IOS/iPhonePackager/iPhonePackager.cs
@@ -92,6 +92,7 @@ namespace iPhonePackager
 		static public BackgroundWorker BGWorker = null;
 		static public int ProgressIndex = 0;
 		static public string AdditionalCommandline = "";
+		static public bool IsClient = false;
 
 		static public void UpdateStatus(string Line)
 		{
@@ -398,6 +399,11 @@ namespace iPhonePackager
 							if (Arguments.Length > ArgIndex + 1)
 							{
 								SchemeConfiguration = Arguments[++ArgIndex];
+
+								if(SchemeConfiguration.EndsWith("Client"))
+								{
+									IsClient = true;
+								}
 							}
 							else
 							{
diff --git a/Engine/Source/Programs/UnrealBuildTool/Platform/IOS/IOSToolChain.cs b/Engine/Source/Programs/UnrealBuildTool/Platform/IOS/IOSToolChain.cs
--- a/Engine/Source/Programs/UnrealBuildTool/Platform/IOS/IOSToolChain.cs
+++ b/Engine/Source/Programs/UnrealBuildTool/Platform/IOS/IOSToolChain.cs
@@ -1521,7 +1521,7 @@ namespace UnrealBuildTool
 				else
 				{
 					UnrealBuildTool.GenerateProjectFiles(new XcodeProjectFileGenerator(Target.ProjectFile), new string[] { "-platforms=" + (Target.Platform == UnrealTargetPlatform.IOS ? "IOS" : "TVOS"), "-NoIntellIsense", (Target.Platform == UnrealTargetPlatform.IOS ? "-iosdeployonly" : "-tvosdeployonly"), "-ignorejunk", String.Format("-project={0}", Target.ProjectFile), "-game" });
-					XcodeWorkspaceDir = DirectoryReference.Combine(Target.ProjectDirectory, String.Format("{0}_{1}.xcworkspace", AppName, (Target.Platform == UnrealTargetPlatform.IOS ? "IOS" : "TVOS")));
+					XcodeWorkspaceDir = DirectoryReference.Combine(Target.ProjectDirectory, String.Format("{0}_{1}.xcworkspace", Target.ProjectFile.GetFileNameWithoutExtension(), (Target.Platform == UnrealTargetPlatform.IOS ? "IOS" : "TVOS")));
 				}
 
 				// Make sure it exists
@@ -1630,7 +1630,7 @@ namespace UnrealBuildTool
 					}
 					else
 					{
-						SchemeName = AppName;
+						SchemeName = Target.ProjectFile.GetFileNameWithoutExtension();
 					}
 
 					// code sign the project
