UE-63768 - iOS Device Output Log shows non-UE4 log data by default

diff --git a/Engine/Source/Developer/OutputLog/Private/SDeviceOutputLog.cpp b/Engine/Source/Developer/OutputLog/Private/SDeviceOutputLog.cpp
--- a/Engine/Source/Developer/OutputLog/Private/SDeviceOutputLog.cpp
+++ b/Engine/Source/Developer/OutputLog/Private/SDeviceOutputLog.cpp
@@ -220,11 +220,6 @@ void SDeviceOutputLog::HandleTargetPlatformDeviceDiscovered(ITargetDeviceRef Dis
 	if (DeviceList.IsValidIndex(ExistingEntryIdx))
 	{
 		DeviceList[ExistingEntryIdx]->DeviceWeakPtr = DiscoveredDevice;
-		
-		if (CurrentDevicePtr.IsValid() && CurrentDevicePtr->DeviceId.GetDeviceName() == DeviceList[ExistingEntryIdx]->DeviceId.GetDeviceName())
-		{
-			CurrentDeviceOutputPtr = DiscoveredDevice->CreateDeviceOutputRouter(this);
-		}
 	}
 	else
 	{
diff --git a/Engine/Source/Developer/iOS/IOSTargetPlatform/Private/IOSTargetDeviceOutput.inl b/Engine/Source/Developer/iOS/IOSTargetPlatform/Private/IOSTargetDeviceOutput.inl
--- a/Engine/Source/Developer/iOS/IOSTargetPlatform/Private/IOSTargetDeviceOutput.inl
+++ b/Engine/Source/Developer/iOS/IOSTargetPlatform/Private/IOSTargetDeviceOutput.inl
@@ -106,7 +106,7 @@ inline uint32 FIOSDeviceOutputReaderRunnable::Run(void)
 
 				for (int32 i = 0; i < OutputLines.Num(); ++i)
 				{
-					//if (OutputLines[i].Contains(TEXT("[UE4]"), ESearchCase::CaseSensitive))
+					if (OutputLines[i].Contains(TEXT("[UE4]"), ESearchCase::CaseSensitive))
 					{
 						Output->Serialize(*OutputLines[i], ELogVerbosity::Log, NAME_None);
 					}
diff --git a/Engine/Source/Programs/IOS/MobileDeviceInterface/MobileDevice.cs b/Engine/Source/Programs/IOS/MobileDeviceInterface/MobileDevice.cs
--- a/Engine/Source/Programs/IOS/MobileDeviceInterface/MobileDevice.cs
+++ b/Engine/Source/Programs/IOS/MobileDeviceInterface/MobileDevice.cs
@@ -320,6 +320,7 @@ namespace Manzana
 		int SecureStartService(TypedPtr<AppleMobileDeviceConnection> device, TypedPtr<CFString> serviceName, int flagsPassInZero, ref IntPtr handle);
 		string ServiceConnectionReceive(IntPtr handle);
 		int SocketSend(IntPtr fd, IntPtr buf, int count);
+		void SocketClose(IntPtr fd);
 		int USBMuxConnectByPort(int DeviceID, short iPort, ref IntPtr outHandle);
 		int RestoreRegisterForDeviceNotifications(
 			DeviceRestoreNotificationCallback dfu_connect,
@@ -474,6 +475,13 @@ namespace Manzana
 			return (int)write(fd, buf, (ulong)count);
 		}
 
+		[DllImport("libikvm-native.dylib", SetLastError = true)]
+		public static extern int close(IntPtr fd);
+		public void SocketClose(IntPtr fd)
+		{
+			close(fd);
+		}
+
 		public int RestoreRegisterForDeviceNotifications(
 			DeviceRestoreNotificationCallback dfu_connect,
 			DeviceRestoreNotificationCallback recovery_connect,
@@ -722,12 +730,31 @@ namespace Manzana
 			{
 				return AMDeviceSecureStartService((IntPtr)device, (IntPtr)serviceName, flagsPassInZero, ref handle);
 			}
-		
+
+			
 			public static string ServiceConnectionReceive(IntPtr handle)
 			{
-				StringBuilder sb = new StringBuilder("", 1024);
-				AMDServiceConnectionReceive(handle, sb, sb.Capacity, 0);
-				return sb.ToString();
+				byte[] Buffer = new byte[1024];
+
+				IntPtr UnmanagedPointer = Marshal.UnsafeAddrOfPinnedArrayElement(Buffer, 0);
+				int Ret = AMDServiceConnectionReceive(handle, UnmanagedPointer, 1023, 0);
+				Buffer[Ret] = 0;
+				int Start = 0;
+				string Result = "";
+				int I;
+				for (I = 0; I < Ret; I++)
+				{
+					if (Buffer[I] == 0)
+					{
+						Result += System.Text.Encoding.UTF8.GetString(Buffer, Start, I - Start);
+						Start = I + 1;
+					}
+				}
+				if (I > Start)
+				{
+					Result += System.Text.Encoding.UTF8.GetString(Buffer, Start, I - Start);
+				}
+				return Result;
 			}
 			
 			public static int UninstallApplication(
@@ -842,7 +869,7 @@ namespace Manzana
 			private extern static int AMDeviceSecureStartService(IntPtr/*AppleMobileDeviceConnection*/ device, IntPtr/*CFString*/ serviceName, int flagsPassInZero, ref IntPtr handle);
 
 			[DllImport(DLLName, CallingConvention = CallingConvention.Cdecl)]
-			private extern static int AMDServiceConnectionReceive(IntPtr handle, StringBuilder buffer, int bufferLen, int unknown);
+			private extern static int AMDServiceConnectionReceive(IntPtr handle, IntPtr buffer, int bufferLen, int unknown);
 
 			[DllImport(DLLName, CallingConvention = CallingConvention.Cdecl)]
 			private extern static int AMDeviceStartSession(IntPtr/*AppleMobileDeviceConnection*/ device);
@@ -1286,6 +1313,13 @@ namespace Manzana
 			return (int)send(fd, buf, count, 0);
 		}
 
+		[DllImport("ws2_32.dll")]
+		public static extern void closesocket(IntPtr fd);
+		public void SocketClose(IntPtr fd)
+		{
+			closesocket(fd);
+		}
+
 		public int RestoreRegisterForDeviceNotifications(
 			DeviceRestoreNotificationCallback dfu_connect,
 			DeviceRestoreNotificationCallback recovery_connect,
@@ -1534,12 +1568,31 @@ namespace Manzana
 			{
 				return AMDeviceSecureStartService((IntPtr)device, (IntPtr)serviceName, flagsPassInZero, ref handle);
 			}
-		
+
+			
 			public static string ServiceConnectionReceive(IntPtr handle)
 			{
-				StringBuilder sb = new StringBuilder(1024);
-				AMDServiceConnectionReceive(handle, sb, sb.Capacity, 0);
-				return sb.ToString();
+				byte[] Buffer = new byte[1024];
+
+				IntPtr UnmanagedPointer = Marshal.UnsafeAddrOfPinnedArrayElement(Buffer, 0);
+				int Ret = AMDServiceConnectionReceive(handle, UnmanagedPointer, 1023, 0);
+				Buffer[Ret] = 0;
+				int Start = 0;
+				string Result = "";
+				int I;
+				for (I = 0; I < Ret; I++)
+				{
+					if (Buffer[I] == 0)
+					{
+						Result += System.Text.Encoding.UTF8.GetString(Buffer, Start, I - Start);
+						Start = I + 1;
+					}
+				}
+				if (I > Start)
+				{
+					Result += System.Text.Encoding.UTF8.GetString(Buffer, Start, I - Start);
+				}
+				return Result;
 			}
 
 			public static int UninstallApplication(
@@ -1654,7 +1707,7 @@ namespace Manzana
 			private extern static int AMDeviceSecureStartService(IntPtr/*AppleMobileDeviceConnection*/ device, IntPtr/*CFString*/ serviceName, int flagsPassInZero, ref IntPtr handle);
 
 			[DllImport(DLLName, CallingConvention = CallingConvention.Cdecl)]
-			private extern static int AMDServiceConnectionReceive(IntPtr handle, StringBuilder buffer, int bufferLen, int unknown);
+			private extern static int AMDServiceConnectionReceive(IntPtr handle, IntPtr buffer, int bufferLen, int unknown);
 
 			[DllImport(DLLName, CallingConvention = CallingConvention.Cdecl)]
 			private extern static int AMDeviceStartSession(IntPtr/*AppleMobileDeviceConnection*/ device);
@@ -2111,6 +2164,13 @@ namespace Manzana
 			return (int)send(fd, buf, count, 0);
 		}
 
+		[DllImport("ws2_32.dll")]
+		public static extern void closesocket(IntPtr fd);
+		public void SocketClose(IntPtr fd)
+		{
+			closesocket(fd);
+		}
+
 		public int RestoreRegisterForDeviceNotifications(
 			DeviceRestoreNotificationCallback dfu_connect,
 			DeviceRestoreNotificationCallback recovery_connect,
@@ -2359,12 +2419,31 @@ namespace Manzana
 			{
 				return AMDeviceSecureStartService((IntPtr)device, (IntPtr)serviceName, flagsPassInZero, ref handle);
 			}
-		
+
+			
 			public static string ServiceConnectionReceive(IntPtr handle)
 			{
-				StringBuilder sb = new StringBuilder(1024);
-				AMDServiceConnectionReceive(handle, sb, sb.Capacity, 0);
-				return sb.ToString();
+				byte[] Buffer = new byte[1024];
+
+				IntPtr UnmanagedPointer = Marshal.UnsafeAddrOfPinnedArrayElement(Buffer, 0);
+				int Ret = AMDServiceConnectionReceive(handle, UnmanagedPointer, 1023, 0);
+				Buffer[Ret] = 0;
+				int Start = 0;
+				string Result = "";
+				int I;
+				for (I = 0; I < Ret; I++)
+				{
+					if (Buffer[I] == 0)
+					{
+						Result += System.Text.Encoding.UTF8.GetString(Buffer, Start, I - Start);
+						Start = I + 1;
+					}
+				}
+				if (I > Start)
+				{
+					Result += System.Text.Encoding.UTF8.GetString(Buffer, Start, I - Start);
+				}
+				return Result;
 			}
 
 			public static int UninstallApplication(
@@ -2479,7 +2558,7 @@ namespace Manzana
 			private extern static int AMDeviceSecureStartService(IntPtr/*AppleMobileDeviceConnection*/ device, IntPtr/*CFString*/ serviceName, int flagsPassInZero, ref IntPtr handle);
 
 			[DllImport(DLLName, CallingConvention = CallingConvention.Cdecl)]
-			private extern static int AMDServiceConnectionReceive(IntPtr handle, StringBuilder buffer, int bufferLen, int unknown);
+			private extern static int AMDServiceConnectionReceive(IntPtr handle, IntPtr buffer, int bufferLen, int unknown);
 
 			[DllImport(DLLName, CallingConvention = CallingConvention.Cdecl)]
 			private extern static int AMDeviceStartSession(IntPtr/*AppleMobileDeviceConnection*/ device);
diff --git a/Engine/Source/Runtime/Core/Private/Apple/ApplePlatformMisc.cpp b/Engine/Source/Runtime/Core/Private/Apple/ApplePlatformMisc.cpp
--- a/Engine/Source/Runtime/Core/Private/Apple/ApplePlatformMisc.cpp
+++ b/Engine/Source/Runtime/Core/Private/Apple/ApplePlatformMisc.cpp
@@ -51,7 +51,7 @@ FString FApplePlatformMisc::GetEnvironmentVariable(const TCHAR* VariableName)
 void FApplePlatformMisc::LowLevelOutputDebugString( const TCHAR *Message )
 {
 	//NsLog will out to all iOS output consoles, instead of just the Xcode console.
-	NSLog(@"%s", TCHAR_TO_UTF8(Message));
+	NSLog(@"[UE4] %s", TCHAR_TO_UTF8(Message));
 }
 
 const TCHAR* FApplePlatformMisc::GetSystemErrorMessage(TCHAR* OutBuffer, int32 BufferCount, int32 Error)
