UE-68538 - IphonePackager GUI operations are not functional

diff --git a/Engine/Source/Programs/IOS/DeploymentServer/DeployTime.cs b/Engine/Source/Programs/IOS/DeploymentServer/DeployTime.cs
--- a/Engine/Source/Programs/IOS/DeploymentServer/DeployTime.cs
+++ b/Engine/Source/Programs/IOS/DeploymentServer/DeployTime.cs
@@ -397,29 +397,56 @@ namespace DeploymentServer
 
 		public void TunnelToDevice(string inDeviceID, String Command)
 		{
-			MobileDeviceInstance targetDevice = null;
+			MobileDeviceInstance TargetDevice = null;
 			
 			PerformActionOnAllDevices(2 * StandardEnumerationDelayMS, delegate (MobileDeviceInstance Device)
 			{
 				if (inDeviceID == Device.DeviceId)
 				{
-					targetDevice = Device;
+					TargetDevice = Device;
 				}
 				return true;
 			});
+			IntPtr TCPService = new IntPtr();
+			if (TargetDevice != null)
+			{
+				TargetDevice.StartTCPRelayService(ref TCPService);
 
-			if (targetDevice != null)
+				int SentDat = TargetDevice.TunnelData(Command, TCPService);
+
+				TargetDevice.StopSyslogService();
+			}
+			else
 			{
-				targetDevice.StartTCPRelayService();
+				ReportIF.Error("Could not find device " + inDeviceID);
+			}
+		}
 
-				int SentDat = targetDevice.TunnelData(Command);
+		public MobileDeviceInstance StartTCPTunnel(string inDeviceID, ref IntPtr TCPService, short Port = 8888)
+		{
+			MobileDeviceInstance targetDevice = null;
 
-				targetDevice.StopSyslogService();
+			PerformActionOnAllDevices(2 * StandardEnumerationDelayMS, delegate (MobileDeviceInstance Device)
+			{
+				if (inDeviceID == Device.DeviceId)
+				{
+					targetDevice = Device;
+				}
+				return true;
+			});
+			if (targetDevice != null)
+			{
+				if (targetDevice.StartTCPRelayService(ref TCPService))
+				{
+					return targetDevice;
+				}
+				ReportIF.Error("Could not start TCP relay to " + inDeviceID);
 			}
 			else
 			{
 				ReportIF.Error("Could not find device " + inDeviceID);
 			}
+			return null;
 		}
 
 		/// <summary>
@@ -432,8 +459,8 @@ namespace DeploymentServer
                 return false;
             }
 
-            // Transfer to all connected devices
-            return PerformActionOnAllDevices(StandardEnumerationDelayMS, delegate(MobileDeviceInstance Device)
+			// Transfer to all connected devices
+			return PerformActionOnAllDevices(StandardEnumerationDelayMS, delegate(MobileDeviceInstance Device)
             {
                 // Transfer the file to the device
                 string DeviceName = Device.DeviceName;
diff --git a/Engine/Source/Programs/IOS/DeploymentServer/Program.cs b/Engine/Source/Programs/IOS/DeploymentServer/Program.cs
--- a/Engine/Source/Programs/IOS/DeploymentServer/Program.cs
+++ b/Engine/Source/Programs/IOS/DeploymentServer/Program.cs
@@ -33,7 +33,9 @@
 		static int ClientCounter = 0;
 		static long TimeOut = 120000;
 		static Stopwatch GlobalTimer = Stopwatch.StartNew();
-
+		static List<String> TCPDevices = new List<string>();
+		static List<int> TCPPorts = new List<int>();
+		static int ParentPID = 0;
 		static string TestStartPath = null;
 
 		class TcpClientInfo
@@ -84,7 +86,8 @@
 			protected string Manifest = "";
 			protected string IpaPath = "";
 			protected string Device = "";
-			protected string Param = "";
+			protected string Param1 = "";
+			protected string Param2 = "";
 			protected bool bKeepAlive = false;
 			protected bool LastResult = false;
 
@@ -127,92 +130,101 @@
 					for (int ArgIndex = Start + 1; ArgIndex < Arguments.Count; ArgIndex++)
 					{
 						string Arg = Arguments[ArgIndex].ToLowerInvariant();
-						if (Arg.StartsWith("-"))
+						switch (Arg)
 						{
-							switch (Arg)
-							{
-								case "-file":
+							case "-file":
+								if (Arguments.Count > ArgIndex + 1)
+								{
+									FileList.Add(Arguments[++ArgIndex]);
+								}
+								else
+								{
+									return false;
+								}
+								break;
+
+							case "-bundle":
+								if (Arguments.Count > ArgIndex + 1)
+								{
+									Bundle = Arguments[++ArgIndex];
+								}
+								else
+								{
+									return false;
+								}
+								break;
+
+							case "-manifest":
+								if (Arguments.Count > ArgIndex + 1)
+								{
+									Manifest = Arguments[++ArgIndex];
+								}
+								else
+								{
+									return false;
+								}
+								break;
+
+							case "-ipa":
+								if (Arguments.Count > ArgIndex + 1)
+								{
+									IpaPath = Arguments[++ArgIndex];
+								}
+								else
+								{
+									return false;
+								}
+								break;
+
+							case "-device":
+								if (Arguments.Count > ArgIndex + 1)
+								{
+									Device = Arguments[++ArgIndex];
+								}
+								else
+								{
+									return false;
+								}
+								break;
+							case "-param":
+								if (Arguments.Count > ArgIndex + 1)
+								{
+									Param1 = Arguments[++ArgIndex];
+								}
+								else
+								{
+									return false;
+								}
+								break;
+							case "-nokeepalive":
+								bKeepAlive = false;
+								break;
+							case "-timeout":
+								{
 									if (Arguments.Count > ArgIndex + 1)
 									{
-										FileList.Add(Arguments[++ArgIndex]);
-									}
-									else
-									{
-										return false;
-									}
-									break;
-
-								case "-bundle":
-									if (Arguments.Count > ArgIndex + 1)
-									{
-										Bundle = Arguments[++ArgIndex];
-									}
-									else
-									{
-										return false;
-									}
-									break;
-
-								case "-manifest":
-									if (Arguments.Count > ArgIndex + 1)
-									{
-										Manifest = Arguments[++ArgIndex];
-									}
-									else
-									{
-										return false;
-									}
-									break;
-
-								case "-ipa":
-									if (Arguments.Count > ArgIndex + 1)
-									{
-										IpaPath = Arguments[++ArgIndex];
-									}
-									else
-									{
-										return false;
-									}
-									break;
-
-								case "-device":
-									if (Arguments.Count > ArgIndex + 1)
-									{
-										Device = Arguments[++ArgIndex];
-									}
-									else
-									{
-										return false;
-									}
-									break;
-								case "-param":
-									if (Arguments.Count > ArgIndex + 1)
-									{
-										Param = Arguments[++ArgIndex];
-									}
-									else
-									{
-										return false;
-									}
-									break;
-								case "-nokeepalive":
-									bKeepAlive = false;
-									break;
-								case "-timeout":
-									{
-										if (Arguments.Count > ArgIndex + 1)
+										long ArgTime = TimeOut;
+										long.TryParse(Arguments[++ArgIndex], out ArgTime);
+										if (ArgTime > 0)
 										{
-											long ArgTime = TimeOut;
-											long.TryParse(Arguments[++ArgIndex], out ArgTime);
-											if (ArgTime > 0)
-											{
-												TimeOut = ArgTime;
-												Console.WriteLine(string.Format("Deployment Server timeout set to {0} (remote)", TimeOut.ToString()));
-											}
+											TimeOut = ArgTime;
+											Console.WriteLine(string.Format("Deployment Server timeout set to {0} (remote)", TimeOut.ToString()));
 										}
-										break;
 									}
-							}
+									break;
+								}
+							default:
+								{
+									if (Param1.Length < 1)
+									{
+										Param1 = Arguments[ArgIndex];
+									}
+									else if (Param2.Length < 1)
+									{
+										Param2 = Arguments[ArgIndex];
+									}
+									break;
+								}
 						}
 					}
 				}
@@ -248,10 +260,13 @@
 									StopTimeout--;
 									if (StopTimeout <= 0)
 									{
-										ForceKillProcesses();
+										Console.WriteLine("Deployment Server Forced Stopping ...");
+										ClientCounter = 0;
+										TCPDevices.Clear();
 										break;
 									}
 								}
+								ForceKillProcesses();
 								break;
 
                             case "backupdocuments":
@@ -294,24 +309,67 @@
 								if (Device.Length < 5)
 								{
 									Console.WriteLine("Device ID not present.");
+									Writer.WriteLine("[command] Device ID not present.");
 								}
-								else if (Param.Length < 1)
+								else if (Param1.Length < 1)
 								{
 									Console.WriteLine("Parameter not present.");
+									Writer.WriteLine("[command] Device ID not present.");
+								}
+								else if (TCPDevices.Contains(Device))
+								{
+									Console.WriteLine("Device already has a TCP connection active.");
+									Writer.WriteLine("[command] Device already has a TCP connection active.");
 								}
 								else
 								{
 									try
 									{
-										DeploymentProxy.Deployer.TunnelToDevice(Device, Param);
+										TCPDevices.Add(Device);
+										IntPtr TCPService = new IntPtr();
+										MobileDeviceInstance targetDevice = DeploymentProxy.Deployer.StartTCPTunnel(Device, ref TCPService);
+										if (targetDevice != null)
+										{
+											int Ret = targetDevice.TunnelData(Param1, TCPService);
+											targetDevice.CloseTunnel(TCPService);
+
+											Writer.WriteLine("[command] Sennt '{0}' bytes.", Ret);
+										}
+										{
+											Writer.WriteLine("[command] Device '{0}' not detected.", Device);
+										}
 									}
 									catch
 									{
 										Console.WriteLine("Errors encountered while tunneling to device.");
+										Writer.WriteLine("[command] Errors encountered while tunneling to device.");
+									}
+									finally
+									{
+										TCPDevices.Remove(Device);
 									}
 								}
 								break;
 
+							case "forward":
+								if (Device.Length < 5)
+								{
+									Console.WriteLine("Device ID not present.");
+								}
+								else if (Param1.Length < 1)
+								{
+									Console.WriteLine("Start TCP port not present.");
+								}
+								else if (Param2.Length < 1)
+								{
+									Console.WriteLine("Destination TCP port not present.");
+								}
+								else
+								{
+									LastResult = ForwardToDevice(Writer, Device, Param1, Param2);
+								}
+								break;
+
 							case "listentodevice":
 								if (Device.Length < 5)
 								{
@@ -370,6 +428,114 @@
 				return LastResult ? 0 : 1;
 			}
 
+			protected bool ForwardToDevice(TextWriter Writer, String Device, String Param1, String Param2)
+			{
+				int Port1 = 0;
+				int Port2 = 0;
+				if (!int.TryParse(Param1, out Port1))
+				{
+					Writer.WriteLine("Invalid start port specified.");
+					return false;
+				}
+				if (!int.TryParse(Param2, out Port2))
+				{
+					Writer.WriteLine("Invalid destination port specified.");
+					return false;
+				}
+				if (TCPDevices.Contains(Device))
+				{
+					Writer.WriteLine("Device {0} already has a TCP connection active.", Device);
+					return false;
+				}
+				if (TCPPorts.Contains(Port1))
+				{
+					Writer.WriteLine("Port {0} is already used as a TCP connection.", Port1);
+					return false;
+				}
+				
+				Thread ProcessClient = new System.Threading.Thread(delegate ()
+				{
+					while (!IsStopping)
+					{
+						try
+						{
+							IntPtr TCPService = new IntPtr();
+							MobileDeviceInstance TargetDevice = null;
+
+							TCPDevices.Add(Device);
+							TCPPorts.Add(Port1);
+
+							TargetDevice = DeploymentProxy.Deployer.StartTCPTunnel(Device, ref TCPService);
+							if (TargetDevice == null)
+							{
+								Console.WriteLine("Cannot connect to device {0} for port forwarding.", Device);
+								break;
+							}
+							Console.WriteLine("Connected to device.");
+
+							TcpListener Server = null;
+							Server = new TcpListener(IPAddress.Any, Port1);
+							Server.Start();
+
+							try
+							{
+								TcpClient Client = Server.AcceptTcpClient();
+								Console.WriteLine("Got TCP connection.");
+								NetworkStream ClStream = Client.GetStream();
+								Byte[] Buffer = new Byte[1024];
+
+								while (!IsStopping)
+								{
+									if (ClStream.DataAvailable)
+									{
+										int Bytes = ClStream.Read(Buffer, 0, Buffer.Length);
+										TargetDevice.TunnelBuffer(Buffer, Bytes, TCPService);
+									}
+									else
+									{
+										if (Client.Client.Poll(10, SelectMode.SelectRead))
+										{
+											Console.WriteLine("TCP disconnected.");
+											break;
+										}
+										System.Threading.Thread.Sleep(100);
+									}
+								}
+							}
+							catch
+							{
+
+							}
+							finally
+							{
+								Console.WriteLine("Port forwarding disconnected.");
+								if (TargetDevice != null)
+								{
+									TargetDevice.CloseTunnel(TCPService);
+								}
+								if (Server != null)
+								{
+									Server.Stop();
+								}
+
+								TCPDevices.Remove(Device);
+								TCPPorts.Remove(Port1);
+							}
+						}
+						catch
+						{
+
+						}
+						finally
+						{
+							
+						}
+					}
+				});
+				ProcessClient.Start();
+				return true;
+			}
+
 			public static bool NeedsVersionCheck(String Command)
 			{
 				switch (Command)
@@ -407,12 +573,13 @@
 			{
 				NewProcess.StartInfo.WorkingDirectory.TrimEnd('/');
 				NewProcess.StartInfo.FileName = NewProcess.StartInfo.WorkingDirectory + "/../../../Build/BatchFiles/Mac/RunMono.sh";
-				NewProcess.StartInfo.Arguments = "\"" + NewProcess.StartInfo.WorkingDirectory + "/DeploymentServer.exe\" -iphonepackager " + Process.GetCurrentProcess().Id.ToString() + " " + NewProcess.StartInfo.WorkingDirectory;
+				NewProcess.StartInfo.Arguments = "\"" + NewProcess.StartInfo.WorkingDirectory + "/DeploymentServer.exe\" server " + " " + NewProcess.StartInfo.WorkingDirectory;
 			}
 			else
 			{
 				NewProcess.StartInfo.WorkingDirectory.TrimEnd('\\');
 				NewProcess.StartInfo.FileName = NewProcess.StartInfo.WorkingDirectory + "\\DeploymentServerLauncher.exe";
+				NewProcess.StartInfo.Arguments = "server";
 			}
 			NewProcess.StartInfo.WindowStyle = ProcessWindowStyle.Hidden;
 			NewProcess.StartInfo.UseShellExecute = true;
@@ -453,7 +620,12 @@
 
 			IpcServerChannel Channel = new IpcServerChannel("iPhonePackager");
 			ChannelServices.RegisterChannel(Channel, false);
-			RemotingConfiguration.RegisterWellKnownServiceType(typeof(DeploymentProxy), "DeploymentServer_PID", WellKnownObjectMode.Singleton);
+			string URI = "DeploymentServer_PID";
+			if (ParentPID > 0)
+			{
+				URI += ParentPID.ToString();
+			}
+			RemotingConfiguration.RegisterWellKnownServiceType(typeof(DeploymentProxy), URI, WellKnownObjectMode.Singleton);
 		}
 
 		protected static void ParseServerParam(string[] Arguments)
@@ -530,6 +702,10 @@
 			System.Threading.Thread ProcessClient = null;
 			try
 			{
+				OutSm = new FileStream("DeploymentServer.log", FileMode.Create, FileAccess.Write);
+				Writer = new StreamWriter(OutSm);
+				Console.SetOut(Writer);
+
 				CreateDeploymentInterface();
 				long.TryParse(ConfigurationManager.AppSettings["DSTimeOut"], out TimeOut);
 				if (TimeOut < 30000)
@@ -538,9 +714,7 @@
 				}
 				Server = new TcpListener(IPAddress.Any, Port);
 				Server.Start();
-				OutSm = new FileStream("DeploymentServer.log", FileMode.Create, FileAccess.Write);
-				Writer = new StreamWriter(OutSm);
-				Console.SetOut(Writer);
+				
 
 				ParseServerParam(Args);
 				Console.WriteLine(string.Format("Deployment Server listening to port {0}", Port.ToString()));
@@ -577,6 +751,10 @@
 					System.Threading.Thread.Sleep(50);
 					Writer.Flush();
 					OutSm.Flush();
+					if (TCPDevices.Count > 0)
+					{
+						GlobalTimer.Restart();
+					}
 					if (ClientCounter <= 0)
 					{
 						if (IsStopping)
@@ -619,6 +797,8 @@
 				{
 					Server.Stop();
 				}
+				Console.WriteLine("Deployment Server Stopped.");
+				Console.SetOut(OldConsole);
 				if (Writer != null)
 				{
 					Writer.Close();
@@ -628,30 +808,66 @@
 					OutSm.Close();
 				}
 			}
-			Console.SetOut(OldConsole);
-			Console.WriteLine("Deployment Server Stopped.");
+			
+			Environment.Exit(0);
 		}
 
 		static void RunLocalInstance(string[] Args)
 		{
+			TextWriter OldConsole = Console.Out;
+			FileStream OutSm = null;
+			TextWriter Writer = null;
+			OutSm = new FileStream("SADeploymentServer.log", FileMode.Create, FileAccess.Write);
+			Writer = new StreamWriter(OutSm);
+			Console.SetOut(Writer);
+			Console.WriteLine("ParentPID: " + ParentPID.ToString());
+
 			// running as one instance only
 			CreateDeploymentInterface();
 			TcpClientInfo LocalClientInfo = new TcpClientInfo();
 			List<string> Arguments = new List<string>();
 			Arguments.AddRange(Args);
 
-			if (LocalClientInfo.ParseCommand(Arguments))
+			if (ParentPID > 0)
 			{
-				LocalClientInfo.HasCommand = true;
-				Console.WriteLine("Running as local instance");
-				LocalClientInfo.RunCommand(Console.Out);
-				while (LocalClientInfo.IsStillRunning)
+				Writer.Flush();
+				try
 				{
-					Thread.Sleep(100);
+					Process ParentProcess = Process.GetProcessById(ParentPID);
+					while (!ParentProcess.HasExited)
+					{
+						CoreFoundationRunLoop.RunLoopRunInMode(CoreFoundationRunLoop.kCFRunLoopDefaultMode(), 1.0, 0);
+					}
 				}
-				bool RetCode = LocalClientInfo.GetLastResult;
+				catch (System.Exception Ex)
+				{
+					Console.WriteLine(Ex.Message);
+				}
 			}
+			else
+			{
+				if (LocalClientInfo.ParseCommand(Arguments))
+				{
+					LocalClientInfo.HasCommand = true;
+					Console.WriteLine("Running as local instance");
+					LocalClientInfo.RunCommand(Console.Out);
+					while (LocalClientInfo.IsStillRunning)
+					{
+						Thread.Sleep(100);
+					}
+					bool RetCode = LocalClientInfo.GetLastResult;
+				}
+			}
 			Console.WriteLine("Deployment Server ended local instance.");
+			Console.SetOut(OldConsole);
+			if (Writer != null)
+			{
+				Writer.Close();
+			}
+			if (OutSm != null)
+			{
+				OutSm.Close();
+			}
 		}
 
 		/**
@@ -774,6 +990,7 @@
 					{
 						Console.WriteLine(Response);
 					}
+					Thread.Sleep(10);
 				}
 			}
 			catch (Exception e)
@@ -812,7 +1029,10 @@
 				Console.WriteLine("\t enumerate");
 				Console.WriteLine("\t listdevices");
 				Console.WriteLine("\t listentodevice");
+				Console.WriteLine("\t command");
+				Console.WriteLine("\t forward");
 				Console.WriteLine("\t -iphonepackager");
+				Console.WriteLine("\t server");
 				Console.WriteLine("Valid Parameters:");
 				Console.WriteLine("\t -file <filename>");
 				Console.WriteLine("\t -bundle <bundle name>");
@@ -821,13 +1041,14 @@
 				Console.WriteLine("\t -device <device ID>");
 				Console.WriteLine("\t -nokeepalive");
 				Console.WriteLine("\t -timeout <miliseconds>");
+				Console.WriteLine("\t -param <string parameter to be used for command>");
 				Console.WriteLine("");
 
 				return 0;
 			}
 			try
 			{
-				if (LocalCommand != "stop")
+				if (LocalCommand != "stop" && LocalCommand != "-iphonepackager")
 				{
 					bool bCreatedMutex = false;
 					String MutexName = "Global\\DeploymentServer_Mutex_RestartNotAllowed";
@@ -854,15 +1075,20 @@
 				Port = DefaultPort;
 			}
 
-			TcpClient IsServiceRunning = IsServiceRegistered();
-
 			// parrent ID not needed anymore
-			if (Args[0].Equals("-iphonepackager"))
+			if (Args[0].Equals("server"))
 			{
+				TcpClient IsServiceRunning = IsServiceRegistered();
 				ServerLoop(IsServiceRunning, Args);
 			}
+			else if (Args[0].Equals("-iphonepackager"))
+			{
+				ParentPID = int.Parse(Args[1]);
+				RunLocalInstance(Args);
+			}
 			else
 			{
+				TcpClient IsServiceRunning = IsServiceRegistered();
 				ClientLoop(IsServiceRunning, Args, LocalCommand);
 			}
 
@@ -926,7 +1152,7 @@
 						//Console.WriteLine("Looping [{0}]", localID);
 						if (ClientInfo.HasCommand && !IsStopping)
 						{
-							Console.WriteLine("Got command [{0}]", localID);
+							//Console.WriteLine("Got command [{0}]", localID);
 							if (!IsRunningCommand && !ClientInfo.IsStillRunning)
 							{
 								ClientInfo.HasCommand = false;
diff --git a/Engine/Source/Programs/IOS/MobileDeviceInterface/MobileDeviceInstance.cs b/Engine/Source/Programs/IOS/MobileDeviceInterface/MobileDeviceInstance.cs
--- a/Engine/Source/Programs/IOS/MobileDeviceInterface/MobileDeviceInstance.cs
+++ b/Engine/Source/Programs/IOS/MobileDeviceInterface/MobileDeviceInstance.cs
@@ -187,7 +187,6 @@ namespace Manzana
         internal TypedPtr<AFCCommConnection> AFCCommsHandle;
         internal IntPtr hService;
         internal IntPtr hSyslogService;
-		internal IntPtr hSysTCPService = new IntPtr();
 		internal IntPtr hInstallService;
         public bool connected;
         private string current_directory;
@@ -1388,7 +1387,7 @@ namespace Manzana
 			return true;
 		}
 
-		public bool StartTCPRelayService()
+		public bool StartTCPRelayService(ref IntPtr TCPService, short Port = 8888)
 		{
 			Console.WriteLine("Connecting");
 			if (MobileDevice.DeviceImpl.Connect(iPhoneHandle) != 0)
@@ -1397,12 +1396,10 @@ namespace Manzana
 				return false;
 			}
 
-			Thread.Sleep(1000);
-			short IPhonePort = 8888;
+			Thread.Sleep(100);
 			int ConnectionID = MobileDevice.DeviceImpl.GetConnectionID(iPhoneHandle);
-			short NetPort = IPAddress.HostToNetworkOrder(IPhonePort);
-			//Console.WriteLine("USBMuxConnectByPort connecting to " + connectionID.ToString() + " port " + iPhonePort.ToString() + " htons " + NetPort.ToString());
-			int R = MobileDevice.DeviceImpl.USBMuxConnectByPort(ConnectionID, NetPort, ref hSysTCPService);
+			short NetPort = IPAddress.HostToNetworkOrder(Port);
+			int R = MobileDevice.DeviceImpl.USBMuxConnectByPort(ConnectionID, NetPort, ref TCPService);
 			if (R != 0)
 			{
 				Console.WriteLine("Connect: Couldn't usb mux by port " + R.ToString());
@@ -1423,16 +1420,14 @@ namespace Manzana
 			MobileDevice.DeviceImpl.Disconnect(iPhoneHandle);
 		}
 
-		public int TunnelData(String Buffer)
+		public int TunnelData(String Buffer, IntPtr TCPService)
 		{
 			int Ret = 0;
 			try
 			{
 				int BufLength = Buffer.Length;
 				IntPtr intPtr_aux = Marshal.StringToHGlobalAnsi(Buffer);
-				//Console.WriteLine("sending data over tcp, len: " + BufLength.ToString());
-				Ret = MobileDevice.DeviceImpl.SocketSend(hSysTCPService, intPtr_aux, BufLength);
-				//Console.WriteLine("send returned: " + Ret.ToString());
+				Ret = MobileDevice.DeviceImpl.SocketSend(TCPService, intPtr_aux, BufLength);
 			}
 			catch (Exception e)
 			{
@@ -1441,6 +1436,33 @@ namespace Manzana
 			return Ret;
 		}
 
+		public int TunnelBuffer(Byte[] Buffer, int Length, IntPtr TCPService)
+		{
+			int Ret = 0;
+			try
+			{
+				IntPtr intPtr_aux = Marshal.UnsafeAddrOfPinnedArrayElement(Buffer, 0);
+				Ret = MobileDevice.DeviceImpl.SocketSend(TCPService, intPtr_aux, Length);
+			}
+			catch (Exception e)
+			{
+				Console.WriteLine("Exception: {0}", e);
+			}
+			return Ret;
+		}
+
+		public void CloseTunnel(IntPtr TCPService)
+		{
+			try
+			{
+				MobileDevice.DeviceImpl.SocketClose(TCPService);
+			}
+			catch (Exception e)
+			{
+				Console.WriteLine("Exception: {0}", e);
+			}
+		}
+
 		#region Private Methods
 		public bool ConnectToPhone()
         {
diff --git a/Engine/Source/Programs/IOS/iPhonePackager/DeployTime.cs b/Engine/Source/Programs/IOS/iPhonePackager/DeployTime.cs
--- a/Engine/Source/Programs/IOS/iPhonePackager/DeployTime.cs
+++ b/Engine/Source/Programs/IOS/iPhonePackager/DeployTime.cs
@@ -245,20 +245,13 @@ namespace iPhonePackager
 		{
 			if (DeployTimeInstance == null)
 			{
-				DeployTimeInstance = (DeploymentInterface)Activator.GetObject(
-				  typeof(DeploymentInterface),
-				  @"ipc://iPhonePackager/DeploymentServer_PID");
-				if (DeployTimeInstance == null)
+				if (DeploymentServerProcess == null)
 				{
-					if (DeploymentServerProcess == null)
-					{
-						DeploymentServerProcess = CreateDeploymentServerProcess();
-					}
+					DeploymentServerProcess = CreateDeploymentServerProcess();
 				}
-
 				DeployTimeInstance = (DeploymentInterface)Activator.GetObject(
-				  typeof(DeploymentInterface),
-				  @"ipc://iPhonePackager/DeploymentServer_PID");
+					typeof(DeploymentInterface),
+					@"ipc://iPhonePackager/DeploymentServer_PID" + Process.GetCurrentProcess().Id.ToString());
 			}
 
 			if (DeployTimeInstance == null)
@@ -282,12 +275,13 @@ namespace iPhonePackager
 			{
 				NewProcess.StartInfo.WorkingDirectory.TrimEnd('/');
 				NewProcess.StartInfo.FileName = "../../../Build/BatchFiles/Mac/RunMono.sh";
-				NewProcess.StartInfo.Arguments = "\"" + NewProcess.StartInfo.WorkingDirectory + "/DeploymentServerLauncher.exe\"";
+				NewProcess.StartInfo.Arguments = "\"" + NewProcess.StartInfo.WorkingDirectory + "/DeploymentServer.exe\" -iphonepackager " + Process.GetCurrentProcess().Id.ToString();
 			}
 			else
 			{
 				NewProcess.StartInfo.WorkingDirectory.TrimEnd('\\');
-				NewProcess.StartInfo.FileName = NewProcess.StartInfo.WorkingDirectory + "\\DeploymentServerLauncher.exe";
+				NewProcess.StartInfo.FileName = NewProcess.StartInfo.WorkingDirectory + "\\DeploymentServer.exe";
+				NewProcess.StartInfo.Arguments = "-iphonepackager " + Process.GetCurrentProcess().Id.ToString();
 			}
 			NewProcess.StartInfo.WindowStyle = ProcessWindowStyle.Hidden;
 			NewProcess.StartInfo.UseShellExecute = true;
@@ -295,7 +289,7 @@ namespace iPhonePackager
 			try
 			{
 				NewProcess.Start();
-				System.Threading.Thread.Sleep(500);
+				System.Threading.Thread.Sleep(1000);
 			}
 			catch (System.Exception ex)
 			{
