UE-68538 - additional emergency fix

diff --git a/Engine/Source/Programs/IOS/DeploymentServerLauncher/Launcher.cs b/Engine/Source/Programs/IOS/DeploymentServerLauncher/Launcher.cs
--- a/Engine/Source/Programs/IOS/DeploymentServerLauncher/Launcher.cs
+++ b/Engine/Source/Programs/IOS/DeploymentServerLauncher/Launcher.cs
@@ -26,7 +26,7 @@ namespace DeploymentServerLauncher
 				NewProcess.StartInfo.WorkingDirectory = WorkingDirectory;
 				NewProcess.StartInfo.WorkingDirectory.TrimEnd('/');
 				NewProcess.StartInfo.FileName = NewProcess.StartInfo.WorkingDirectory + "/../../../Build/BatchFiles/Mac/RunMono.sh";
-				NewProcess.StartInfo.Arguments = "\"" + NewProcess.StartInfo.WorkingDirectory + "/DeploymentServer.exe\" -iphonepackager " + Process.GetCurrentProcess().Id.ToString() + " " + NewProcess.StartInfo.WorkingDirectory;
+				NewProcess.StartInfo.Arguments = "\"" + NewProcess.StartInfo.WorkingDirectory + "/DeploymentServer.exe\" server " + Process.GetCurrentProcess().Id.ToString() + " " + NewProcess.StartInfo.WorkingDirectory;
 
 				NewProcess.StartInfo.WindowStyle = ProcessWindowStyle.Hidden;
 				NewProcess.StartInfo.UseShellExecute = true;
@@ -47,7 +47,7 @@ namespace DeploymentServerLauncher
 				Domaininfo.ConfigurationFile = Path.GetDirectoryName(Assembly.GetExecutingAssembly().Location) + "\\DeploymentServer.exe.config";
 				WorkingDirectory.TrimEnd('\\');
 				UATExecutable = WorkingDirectory + "\\DeploymentServer.exe";
-				CommandLine = "-iphonepackager " + Process.GetCurrentProcess().Id.ToString() + " " + WorkingDirectory;
+				CommandLine = "server " + Process.GetCurrentProcess().Id.ToString() + " " + WorkingDirectory;
 				Domaininfo.ShadowCopyFiles = "true";
 
 				// Create application domain setup information.
