UE-68719 - iPhonePackager no longer sends output to log window

diff --git a/Engine/Source/Programs/IOS/DeploymentServer/Program.cs b/Engine/Source/Programs/IOS/DeploymentServer/Program.cs
--- a/Engine/Source/Programs/IOS/DeploymentServer/Program.cs
+++ b/Engine/Source/Programs/IOS/DeploymentServer/Program.cs
@@ -33,11 +33,18 @@ namespace DeploymentServer
 		static int ClientCounter = 0;
 		static long TimeOut = 120000;
 		static Stopwatch GlobalTimer = Stopwatch.StartNew();
-		static List<String> TCPDevices = new List<string>();
-		static List<int> TCPPorts = new List<int>();
 		static int ParentPID = 0;
 		static string TestStartPath = null;
 
+		class TCPPortForwarding
+		{
+			public String DeviceID;
+			public int TCPPort;
+			public int DevicePort;
+		}
+
+		static List<TCPPortForwarding> TCPForwarding = new List<TCPPortForwarding>();
+
 		class TcpClientInfo
 		{
 			public bool HasCommand
@@ -262,7 +269,7 @@ namespace DeploymentServer
 									{
 										Console.WriteLine("Deployment Server Forced Stopping ...");
 										ClientCounter = 0;
-										TCPDevices.Clear();
+										TCPForwarding.Clear();
 										break;
 									}
 								}
@@ -316,16 +323,10 @@ namespace DeploymentServer
 									Console.WriteLine("Parameter not present.");
 									Writer.WriteLine("[command] Device ID not present.");
 								}
-								else if (TCPDevices.Contains(Device))
-								{
-									Console.WriteLine("Device already has a TCP connection active.");
-									Writer.WriteLine("[command] Device already has a TCP connection active.");
-								}
 								else
 								{
 									try
 									{
-										TCPDevices.Add(Device);
 										IntPtr TCPService = new IntPtr();
 										MobileDeviceInstance targetDevice = DeploymentProxy.Deployer.StartTCPTunnel(Device, ref TCPService);
 										if (targetDevice != null)
@@ -344,10 +345,6 @@ namespace DeploymentServer
 										Console.WriteLine("Errors encountered while tunneling to device.");
 										Writer.WriteLine("[command] Errors encountered while tunneling to device.");
 									}
-									finally
-									{
-										TCPDevices.Remove(Device);
-									}
 								}
 								break;
 
@@ -369,6 +366,13 @@ namespace DeploymentServer
 									LastResult = ForwardToDevice(Writer, Device, Param1, Param2);
 								}
 								break;
+							case "listforwarding":
+								foreach (TCPPortForwarding P in TCPForwarding)
+								{
+									Writer.WriteLine("{0}\r{1}\r{2}", P.DeviceID, P.TCPPort, P.DevicePort);
+								}
+								Writer.WriteLine("");
+								break;
 
 							case "listentodevice":
 								if (Device.Length < 5)
@@ -428,6 +432,22 @@ namespace DeploymentServer
 				return LastResult ? 0 : 1;
 			}
 
+			protected bool IsForwardingInUse(String Device, int Port1, int Port2)
+			{
+				foreach(TCPPortForwarding P in TCPForwarding)
+				{
+					if (P.DevicePort == Port2 && P.DeviceID == Device)
+					{
+						return true;
+					}
+					if (P.TCPPort == Port1)
+					{
+						return true;
+					}
+				}
+				return false;
+			}
+
 			protected bool ForwardToDevice(TextWriter Writer, String Device, String Param1, String Param2)
 			{
 				int Port1 = 0;
@@ -442,14 +462,9 @@ namespace DeploymentServer
 					Writer.WriteLine("Invalid destination port specified.");
 					return false;
 				}
-				if (TCPDevices.Contains(Device))
-				{
-					Writer.WriteLine("Device {0} already has a TCP connection active.", Device);
-					return false;
-				}
-				if (TCPPorts.Contains(Port1))
+				if (IsForwardingInUse(Device, Port1, Port2))
 				{
-					Writer.WriteLine("Port {0} is already used as a TCP connection.", Port1);
+					Writer.WriteLine("Device {0} already has a TCP connection on port {1} or Port {2} is already in use.", Device, Param1, Param2);
 					return false;
 				}
 				
@@ -457,21 +472,17 @@ namespace DeploymentServer
 				{
 					while (!IsStopping)
 					{
+						TCPPortForwarding TcpFW = null;
 						try
 						{
 							IntPtr TCPService = new IntPtr();
 							MobileDeviceInstance TargetDevice = null;
 
-							TCPDevices.Add(Device);
-							TCPPorts.Add(Port1);
-
-							TargetDevice = DeploymentProxy.Deployer.StartTCPTunnel(Device, ref TCPService);
-							if (TargetDevice == null)
-							{
-								Console.WriteLine("Cannot connect to device {0} for port forwarding.", Device);
-								break;
-							}
-							Console.WriteLine("Connected to device.");
+							TcpFW = new TCPPortForwarding();
+							TcpFW.DeviceID = Device;
+							TcpFW.TCPPort = Port1;
+							TcpFW.DevicePort = Port2;
+							TCPForwarding.Add(TcpFW);
 
 							TcpListener Server = null;
 							Server = new TcpListener(IPAddress.Any, Port1);
@@ -482,6 +493,15 @@ namespace DeploymentServer
 								TcpClient Client = Server.AcceptTcpClient();
 								Console.WriteLine("Got TCP connection.");
 								NetworkStream ClStream = Client.GetStream();
+
+								TargetDevice = DeploymentProxy.Deployer.StartTCPTunnel(Device, ref TCPService);
+								if (TargetDevice == null)
+								{
+									Console.WriteLine("Cannot connect to device {0} for port forwarding.", Device);
+									break;
+								}
+								Console.WriteLine("Connected to device.");
+
 								Byte[] Buffer = new Byte[1024];
 
 								while (!IsStopping)
@@ -517,9 +537,6 @@ namespace DeploymentServer
 								{
 									Server.Stop();
 								}
-
-								TCPDevices.Remove(Device);
-								TCPPorts.Remove(Port1);
 							}
 						}
 						catch
@@ -528,7 +545,10 @@ namespace DeploymentServer
 						}
 						finally
 						{
-							
+							if (TcpFW != null)
+							{
+								TCPForwarding.Remove(TcpFW);
+							}
 						}
 					}
 				});
@@ -751,7 +771,7 @@ namespace DeploymentServer
 					System.Threading.Thread.Sleep(50);
 					Writer.Flush();
 					OutSm.Flush();
-					if (TCPDevices.Count > 0)
+					if (TCPForwarding.Count > 0)
 					{
 						GlobalTimer.Restart();
 					}
@@ -814,14 +834,6 @@ namespace DeploymentServer
 
 		static void RunLocalInstance(string[] Args)
 		{
-			TextWriter OldConsole = Console.Out;
-			FileStream OutSm = null;
-			TextWriter Writer = null;
-			OutSm = new FileStream("SADeploymentServer.log", FileMode.Create, FileAccess.Write);
-			Writer = new StreamWriter(OutSm);
-			Console.SetOut(Writer);
-			Console.WriteLine("ParentPID: " + ParentPID.ToString());
-
 			// running as one instance only
 			CreateDeploymentInterface();
 			TcpClientInfo LocalClientInfo = new TcpClientInfo();
@@ -830,7 +842,6 @@ namespace DeploymentServer
 
 			if (ParentPID > 0)
 			{
-				Writer.Flush();
 				try
 				{
 					Process ParentProcess = Process.GetProcessById(ParentPID);
@@ -858,16 +869,6 @@ namespace DeploymentServer
 					bool RetCode = LocalClientInfo.GetLastResult;
 				}
 			}
-			Console.WriteLine("Deployment Server ended local instance.");
-			Console.SetOut(OldConsole);
-			if (Writer != null)
-			{
-				Writer.Close();
-			}
-			if (OutSm != null)
-			{
-				OutSm.Close();
-			}
 		}
 
 		/**
diff --git a/Engine/Source/Programs/IOS/iPhonePackager/DeployTime.cs b/Engine/Source/Programs/IOS/iPhonePackager/DeployTime.cs
--- a/Engine/Source/Programs/IOS/iPhonePackager/DeployTime.cs
+++ b/Engine/Source/Programs/IOS/iPhonePackager/DeployTime.cs
@@ -283,8 +283,7 @@ namespace iPhonePackager
 				NewProcess.StartInfo.FileName = NewProcess.StartInfo.WorkingDirectory + "\\DeploymentServer.exe";
 				NewProcess.StartInfo.Arguments = "-iphonepackager " + Process.GetCurrentProcess().Id.ToString();
 			}
-			NewProcess.StartInfo.WindowStyle = ProcessWindowStyle.Hidden;
-			NewProcess.StartInfo.UseShellExecute = true;
+			NewProcess.StartInfo.UseShellExecute = false;
 
 			try
 			{
