UE-68822 - DeploymentServer crash in IPP when DeploymentServerLauncher.exe is already running

diff --git a/Engine/Source/Programs/IOS/DeploymentServer/Program.cs b/Engine/Source/Programs/IOS/DeploymentServer/Program.cs
--- a/Engine/Source/Programs/IOS/DeploymentServer/Program.cs
+++ b/Engine/Source/Programs/IOS/DeploymentServer/Program.cs
@@ -1,4 +1,4 @@
-﻿/**
+/**
  * Copyright 1998-2019 Epic Games, Inc. All Rights Reserved.
  */
 
@@ -636,10 +636,11 @@ namespace DeploymentServer
 
 		static void CreateDeploymentInterface()
 		{
-			DeploymentProxy.Deployer = new DeploymentImplementation();
-
-			IpcServerChannel Channel = new IpcServerChannel("iPhonePackager");
-			ChannelServices.RegisterChannel(Channel, false);
+			if (ChannelServices.GetChannel("iPhonePackager") == null)
+			{
+				IpcServerChannel Channel = new IpcServerChannel("iPhonePackager");
+				ChannelServices.RegisterChannel(Channel, false);
+			}
 			string URI = "DeploymentServer_PID";
 			if (ParentPID > 0)
 			{
@@ -726,7 +727,7 @@ namespace DeploymentServer
 				Writer = new StreamWriter(OutSm);
 				Console.SetOut(Writer);
 
-				CreateDeploymentInterface();
+				DeploymentProxy.Deployer = new DeploymentImplementation();
 				long.TryParse(ConfigurationManager.AppSettings["DSTimeOut"], out TimeOut);
 				if (TimeOut < 30000)
 				{
@@ -835,6 +836,7 @@ namespace DeploymentServer
 		static void RunLocalInstance(string[] Args)
 		{
 			// running as one instance only
+			DeploymentProxy.Deployer = new DeploymentImplementation();
 			CreateDeploymentInterface();
 			TcpClientInfo LocalClientInfo = new TcpClientInfo();
 			List<string> Arguments = new List<string>();
