UE-68820 Failure to install IPAs to iOS devices

#4.2.2
#iOS

Adding a check for existing iCloud keys in the provisioning profile, not to add an empty one by mistake

diff --git a/Engine/Source/Programs/IOS/iPhonePackager/MobileProvisionUtilities.cs b/Engine/Source/Programs/IOS/iPhonePackager/MobileProvisionUtilities.cs
--- a/Engine/Source/Programs/IOS/iPhonePackager/MobileProvisionUtilities.cs
+++ b/Engine/Source/Programs/IOS/iPhonePackager/MobileProvisionUtilities.cs
@@ -318,22 +318,25 @@ namespace iPhonePackager
 			// must have CloudKit and CloudDocuments for com.apple.developer.icloud-services
 			// otherwise the game will not be listed in the Settings->iCloud apps menu on the device
 			{
-				List<string> ServicesGroups = XCentPList.GetArray("com.apple.developer.icloud-services", "string");
-
-				if (ServicesGroups.Count == 0 || !ServicesGroups[0].StartsWith("Cloud"))
+				if (XCentPList.HasKey("com.apple.developer.icloud-services"))
 				{
-					ServicesGroups.Clear();
-
-					string ServicesString;
-					XCentPList.GetString("com.apple.developer.icloud-services", out ServicesString);
+					List<string> ServicesGroups = XCentPList.GetArray("com.apple.developer.icloud-services", "string");
 
-					if (ServicesString.Equals("*"))
+					if (ServicesGroups.Count == 0 || !ServicesGroups[0].StartsWith("Cloud"))
 					{
-						ServicesGroups.Add("CloudKit");
-						ServicesGroups.Add("CloudDocuments");
+						ServicesGroups.Clear();
+
+						string ServicesString;
+						XCentPList.GetString("com.apple.developer.icloud-services", out ServicesString);
+
+						if (ServicesString.Equals("*"))
+						{
+							ServicesGroups.Add("CloudKit");
+							ServicesGroups.Add("CloudDocuments");
+						}
 					}
+					XCentPList.SetValueForKey("com.apple.developer.icloud-services", ServicesGroups);
 				}
-				XCentPList.SetValueForKey("com.apple.developer.icloud-services", ServicesGroups);
 			}
 
 			return XCentPList.SaveToString();
