UE-69055 - //UE4/Main - Build Tools CS - Processor architecture mismatch between project "MSIL" and "AMD64"

diff --git a/Engine/Source/Programs/IOS/DeploymentServer/DeploymentServer.csproj b/Engine/Source/Programs/IOS/DeploymentServer/DeploymentServer.csproj
--- a/Engine/Source/Programs/IOS/DeploymentServer/DeploymentServer.csproj
+++ b/Engine/Source/Programs/IOS/DeploymentServer/DeploymentServer.csproj
@@ -60,12 +60,12 @@
     <DefineConstants>TRACE</DefineConstants>
     <Optimize>true</Optimize>
     <DebugType>pdbonly</DebugType>
-    <PlatformTarget>x64</PlatformTarget>
+    <PlatformTarget>AnyCPU</PlatformTarget>
     <ErrorReport>prompt</ErrorReport>
     <CodeAnalysisRuleSet>MinimumRecommendedRules.ruleset</CodeAnalysisRuleSet>
     <CodeAnalysisIgnoreBuiltInRuleSets>false</CodeAnalysisIgnoreBuiltInRuleSets>
     <CodeAnalysisIgnoreBuiltInRules>false</CodeAnalysisIgnoreBuiltInRules>
-    <Prefer32Bit>true</Prefer32Bit>
+    <Prefer32Bit>false</Prefer32Bit>
   </PropertyGroup>
   <ItemGroup>
     <Reference Include="System" />
diff --git a/Engine/Source/Programs/IOS/DeploymentServerLauncher/DeploymentServerLauncher.csproj b/Engine/Source/Programs/IOS/DeploymentServerLauncher/DeploymentServerLauncher.csproj
--- a/Engine/Source/Programs/IOS/DeploymentServerLauncher/DeploymentServerLauncher.csproj
+++ b/Engine/Source/Programs/IOS/DeploymentServerLauncher/DeploymentServerLauncher.csproj
@@ -41,7 +41,7 @@
     <AllowUnsafeBlocks>false</AllowUnsafeBlocks>
   </PropertyGroup>
   <PropertyGroup Condition=" '$(Configuration)|$(Platform)' == 'Development|AnyCPU' ">
-    <PlatformTarget>x64</PlatformTarget>
+    <PlatformTarget>AnyCPU</PlatformTarget>
     <DebugType>pdbonly</DebugType>
     <Optimize>true</Optimize>
     <OutputPath>..\..\..\..\Binaries\DotNET\IOS\</OutputPath>
