UE-66204 iOS and tvOS errors in Application Loader to Apple: 'Invalid Code Signing Entitlements' with iCloud

#iOS
#4.22

For distribution builds, the entitlements have a modified syntax comparing to the entitlements section from the mobileprovisioning file.
When signing the ipa, the iPhonePackeger uses the values from the mobileprovisioning file.
Adding filters to validate the entitlements before sending them as parameter for the codesign tool.

diff --git a/Engine/Source/Programs/IOS/iPhonePackager/MobileProvisionUtilities.cs b/Engine/Source/Programs/IOS/iPhonePackager/MobileProvisionUtilities.cs
--- a/Engine/Source/Programs/IOS/iPhonePackager/MobileProvisionUtilities.cs
+++ b/Engine/Source/Programs/IOS/iPhonePackager/MobileProvisionUtilities.cs
@@ -322,14 +322,14 @@ namespace iPhonePackager
 				{
 					List<string> ServicesGroups = XCentPList.GetArray("com.apple.developer.icloud-services", "string");
 
-					if (ServicesGroups.Count == 0 || !ServicesGroups[0].StartsWith("Cloud"))
+					if (ServicesGroups.Count == 0 || !ServicesGroups[0].Contains("*"))
 					{
 						ServicesGroups.Clear();
 
 						string ServicesString;
 						XCentPList.GetString("com.apple.developer.icloud-services", out ServicesString);
 
-						if (ServicesString.Equals("*"))
+						if (ServicesString.Contains("*"))
 						{
 							ServicesGroups.Add("CloudKit");
 							ServicesGroups.Add("CloudDocuments");
@@ -337,6 +337,72 @@ namespace iPhonePackager
 					}
 					XCentPList.SetValueForKey("com.apple.developer.icloud-services", ServicesGroups);
 				}
+
+				// For distribution builds, the entitlements from mobileprovisioning have a modified syntax
+				if (Config.bForDistribution)
+				{
+					if (XCentPList.HasKey("com.apple.developer.ubiquity-kvstore-identifier"))
+					{
+						string UbiquityKvstoreString;
+						XCentPList.GetString("com.apple.developer.ubiquity-kvstore-identifier", out UbiquityKvstoreString);
+
+						if (UbiquityKvstoreString.Contains("*"))
+						{
+							string NewUbiquityKvstoreIdentifier = String.Format("{0}.{1}", TeamIdentifier, CFBundleIdentifier);
+							XCentPList.SetValueForKey("com.apple.developer.ubiquity-kvstore-identifier", NewUbiquityKvstoreIdentifier);
+						}
+					}
+					if (XCentPList.HasKey("com.apple.developer.ubiquity-container-identifiers"))
+					{
+						List<string> UbiquityContainerIdentifiersGroups = XCentPList.GetArray("com.apple.developer.ubiquity-container-identifiers", "string");
+
+						if (UbiquityContainerIdentifiersGroups.Count == 0 || UbiquityContainerIdentifiersGroups[0].Contains("*"))
+						{
+							UbiquityContainerIdentifiersGroups.Clear();
+
+							string UbiquityContainerString;
+							XCentPList.GetString("com.apple.developer.ubiquity-container-identifiers", out UbiquityContainerString);
+
+							if (UbiquityContainerString.Contains("*"))
+							{
+								string NewUbiquityContainerIdentifier = String.Format("{0}.{1}", TeamIdentifier, CFBundleIdentifier);
+								UbiquityContainerIdentifiersGroups.Add(NewUbiquityContainerIdentifier);
+							}
+						}
+
+						XCentPList.SetValueForKey("com.apple.developer.ubiquity-container-identifiers", UbiquityContainerIdentifiersGroups);
+					}
+					if (XCentPList.HasKey("com.apple.developer.associated-domains"))
+					{
+						List<string> AssociatedDomainsGroup = XCentPList.GetArray("com.apple.developer.associated-domains", "string");
+
+						if (AssociatedDomainsGroup.Count == 0 || AssociatedDomainsGroup[0].Contains("*"))
+						{
+							string AssociatedDomainsString;
+							XCentPList.GetString("com.apple.developer.associated-domains", out AssociatedDomainsString);
+
+							if (AssociatedDomainsString.Contains("*"))
+							{
+								XCentPList.RemoveKeyValue("com.apple.developer.associated-domains");
+							}
+						}
+					}
+
+					// remove development keys - generated when the cloudkit container is in development mode
+					XCentPList.RemoveKeyValue("com.apple.developer.icloud-container-development-container-identifiers");
+					if (XCentPList.HasKey("com.apple.developer.icloud-container-environment"))
+					{
+						List<string> ContainerEnvironmentGroup = XCentPList.GetArray("com.apple.developer.icloud-container-environment", "string");
+
+						if (ContainerEnvironmentGroup.Count == 0 || ContainerEnvironmentGroup[0].Equals("Development"))
+						{
+							ContainerEnvironmentGroup.Clear();
+							ContainerEnvironmentGroup.Add("Production");
+						}
+
+						XCentPList.SetValueForKey("com.apple.developer.icloud-container-environment", ContainerEnvironmentGroup);
+					}
+				}
 			}
 
 			return XCentPList.SaveToString();
diff --git a/Engine/Source/Programs/IOS/iPhonePackager/Utilities.cs b/Engine/Source/Programs/IOS/iPhonePackager/Utilities.cs
--- a/Engine/Source/Programs/IOS/iPhonePackager/Utilities.cs
+++ b/Engine/Source/Programs/IOS/iPhonePackager/Utilities.cs
@@ -469,6 +469,32 @@ namespace iPhonePackager
 				return (KeyNode != null);
 			}
 
+			public void RemoveKeyValue(string KeyName)
+			{
+				if (bReadOnly)
+				{
+					throw new AccessViolationException("PList has been set to read only and may not be modified");
+				}
+
+				XmlNode DictionaryNode = Doc.DocumentElement.SelectSingleNode("/plist/dict");
+
+				string PathToKey = String.Format("/plist/dict/key[.='{0}']", KeyName);
+				XmlNode KeyNode = Doc.DocumentElement.SelectSingleNode(PathToKey);
+				if (KeyNode != null && KeyNode.ParentNode != null)
+				{
+					XmlNode ValueNode = KeyNode.NextSibling;
+					//remove value
+					if (ValueNode != null)
+					{
+						ValueNode.RemoveAll();
+						ValueNode.ParentNode.RemoveChild(ValueNode);
+					}
+					//remove key
+					KeyNode.RemoveAll();
+					KeyNode.ParentNode.RemoveChild(KeyNode);
+				}
+			}
+
 			public void SetValueForKey(string KeyName, object Value)
 			{
 				if (bReadOnly)
