Merge conflict from Release-4.22 to UE4/main

diff --git a/Engine/Source/Programs/IOS/DeploymentServer/Program.cs b/Engine/Source/Programs/IOS/DeploymentServer/Program.cs
--- a/Engine/Source/Programs/IOS/DeploymentServer/Program.cs
+++ b/Engine/Source/Programs/IOS/DeploymentServer/Program.cs
@@ -398,6 +398,11 @@ namespace DeploymentServer
 								break;
 						}
 					}
+					
+					catch (IOException)
+					{
+						// we expect this to happen so we don't log it
+					}
 					catch (Exception e)
 					{
 						Console.SetOut(ConsoleOld);
@@ -1165,7 +1170,7 @@ namespace DeploymentServer
 					
 					Writer.WriteLine("[DSDIR]" + TestStartPath);
 					Writer.Flush();
-
+					
 					Byte[] Buffer = new Byte[2048];
 					string Unprocessed = "";
 					ClientInfo = new TcpClientInfo();
@@ -1273,6 +1278,10 @@ namespace DeploymentServer
 						}
 					}
 				}
+				catch (IOException)
+				{
+					// we expect this to happen so we don't log it
+				}
 				catch (Exception e)
 				{
 					Console.WriteLine("Exception: {0}", e);
