UE-69887 - DeploymentServerLauncher.exe does not always quit

diff --git a/Engine/Source/Programs/IOS/DeploymentServer/DeployTime.cs b/Engine/Source/Programs/IOS/DeploymentServer/DeployTime.cs
--- a/Engine/Source/Programs/IOS/DeploymentServer/DeployTime.cs
+++ b/Engine/Source/Programs/IOS/DeploymentServer/DeployTime.cs
@@ -240,7 +240,11 @@ namespace DeploymentServer
                         ReportIF.Error("Timed out while trying to connect to a mobile device.  Make sure one is connected.");
                     }
                 }
-                catch (Exception ex)
+				catch (ThreadAbortException)
+				{
+					// we expect this to happen so we don't log it
+				}
+				catch (Exception ex)
                 {
                     ReportIF.Error(String.Format("Error encountered ('{0}') while trying to connect to a mobile device.  Please verify that iTunes is installed", ex.Message));
                 }
diff --git a/Engine/Source/Programs/IOS/DeploymentServer/Program.cs b/Engine/Source/Programs/IOS/DeploymentServer/Program.cs
--- a/Engine/Source/Programs/IOS/DeploymentServer/Program.cs
+++ b/Engine/Source/Programs/IOS/DeploymentServer/Program.cs
@@ -105,6 +105,7 @@ namespace DeploymentServer
 			protected bool LastResult = false;
 
 			System.Threading.Thread runLoop = null;
+			System.Threading.Thread listenLoop = null;
 			bool bCommandComplete = true;
 
 			public void StopRun()
@@ -115,6 +116,10 @@ namespace DeploymentServer
 					{
 						runLoop.Abort();
 					}
+					if (listenLoop != null)
+					{
+						listenLoop.Abort();
+					}
 				}
 				catch { }
 			}
@@ -398,7 +403,30 @@ namespace DeploymentServer
 								}
 								else
 								{
-									DeploymentProxy.Deployer.ListenToDevice(Device, Writer);
+									listenLoop = new System.Threading.Thread(delegate ()
+									{
+										try
+										{
+											DeploymentProxy.Deployer.ListenToDevice(Device, Writer);
+										}
+										catch (IOException)
+										{
+											// we expect this to happen so we don't log it
+										}
+										catch (ThreadAbortException)
+										{
+											// we expect this to happen so we don't log it
+										}
+										catch (Exception e)
+										{
+											LocalLog("Exception: " + e.ToString());
+										}
+										finally
+										{
+											listenLoop = null;
+										}
+									});
+									listenLoop.Start();
 								}
 								break;
 						}
