Bumped DeploymentServer to Net 4.6.2 for compatibility with MDI

diff --git a/Engine/Source/Programs/IOS/DeploymentServer/DeploymentServer.csproj b/Engine/Source/Programs/IOS/DeploymentServer/DeploymentServer.csproj
--- a/Engine/Source/Programs/IOS/DeploymentServer/DeploymentServer.csproj
+++ b/Engine/Source/Programs/IOS/DeploymentServer/DeploymentServer.csproj
@@ -1,5 +1,5 @@
 ﻿<?xml version="1.0" encoding="utf-8"?>
-<Project ToolsVersion="4.0" DefaultTargets="Build" xmlns="http://schemas.microsoft.com/developer/msbuild/2003">
+<Project ToolsVersion="12.0" DefaultTargets="Build" xmlns="http://schemas.microsoft.com/developer/msbuild/2003">
   <PropertyGroup>
     <Configuration Condition=" '$(Configuration)' == '' ">Debug</Configuration>
     <Platform Condition=" '$(Platform)' == '' ">x86</Platform>
@@ -10,7 +10,7 @@
     <AppDesignerFolder>Properties</AppDesignerFolder>
     <RootNamespace>DeploymentServer</RootNamespace>
     <AssemblyName>DeploymentServer</AssemblyName>
-    <TargetFrameworkVersion>v4.5</TargetFrameworkVersion>
+    <TargetFrameworkVersion>v4.6.2</TargetFrameworkVersion>
     <TargetFrameworkProfile>
     </TargetFrameworkProfile>
     <FileAlignment>512</FileAlignment>
diff --git a/Engine/Source/Programs/IOS/DeploymentServer/app.config b/Engine/Source/Programs/IOS/DeploymentServer/app.config
--- a/Engine/Source/Programs/IOS/DeploymentServer/app.config
+++ b/Engine/Source/Programs/IOS/DeploymentServer/app.config
@@ -1,7 +1,7 @@
 <?xml version="1.0"?>
 <configuration>
 	<startup>
-		<supportedRuntime version="v4.0" sku=".NETFramework,Version=v4.0"/>
+		<supportedRuntime version="v4.0" sku=".NETFramework,Version=v4.6.2"/>
 	</startup>
 	<runtime>
 		<NetFx40_PInvokeStackResilience enabled="1"/>
