Fix issue with iPhonePackager failing on Mac with Mono 6

diff --git a/Engine/Source/Programs/IOS/iPhonePackager/CodeSigning.cs b/Engine/Source/Programs/IOS/iPhonePackager/CodeSigning.cs
--- a/Engine/Source/Programs/IOS/iPhonePackager/CodeSigning.cs
+++ b/Engine/Source/Programs/IOS/iPhonePackager/CodeSigning.cs
@@ -248,12 +248,9 @@ namespace iPhonePackager
 						DateTime EffectiveDate = cert.NotBefore.ToUniversalTime();
 						DateTime ExpirationDate = cert.NotAfter.ToUniversalTime();
 						DateTime Now = DateTime.UtcNow;
-						string Subject = cert.Subject;
-						int SubjStart = Subject.IndexOf("CN=") + 3;
-						int SubjEnd = Subject.IndexOf(",", SubjStart);
-						cert.FriendlyName = Subject.Substring(SubjStart, (SubjEnd - SubjStart));
+
 						bool bCertTimeIsValid = (EffectiveDate < Now) && (ExpirationDate > Now);
-						Program.LogVerbose("CERTIFICATE-Name:{0},Validity:{1},StartDate:{2},EndDate:{3}", cert.FriendlyName, bCertTimeIsValid ? "VALID" : "EXPIRED", EffectiveDate.ToString("o"), ExpirationDate.ToString("o"));
+						Program.LogVerbose("CERTIFICATE-Name:{0},Validity:{1},StartDate:{2},EndDate:{3}", CryptoAdapter.GetFriendlyNameFromCert(cert), bCertTimeIsValid ? "VALID" : "EXPIRED", EffectiveDate.ToString("o"), ExpirationDate.ToString("o"));
 
 						start = CertToolData.IndexOf(header, start);
 					}
@@ -275,7 +272,7 @@ namespace iPhonePackager
 					DateTime Now = DateTime.UtcNow;
 
 					bool bCertTimeIsValid = (EffectiveDate < Now) && (ExpirationDate > Now);
-					Program.LogVerbose("CERTIFICATE-Name:{0},Validity:{1},StartDate:{2},EndDate:{3}", TestCert.FriendlyName, bCertTimeIsValid ? "VALID" : "EXPIRED", EffectiveDate.ToString("o"), ExpirationDate.ToString("o"));
+					Program.LogVerbose("CERTIFICATE-Name:{0},Validity:{1},StartDate:{2},EndDate:{3}", CryptoAdapter.GetFriendlyNameFromCert(TestCert), bCertTimeIsValid ? "VALID" : "EXPIRED", EffectiveDate.ToString("o"), ExpirationDate.ToString("o"));
 				}
 
 				FoundCerts = Store.Certificates.Find(X509FindType.FindBySubjectName, "iPhone Distribution", false);
@@ -287,7 +284,7 @@ namespace iPhonePackager
 					DateTime Now = DateTime.UtcNow;
 
 					bool bCertTimeIsValid = (EffectiveDate < Now) && (ExpirationDate > Now);
-					Program.LogVerbose("CERTIFICATE-Name:{0},Validity:{1},StartDate:{2},EndDate:{3}", TestCert.FriendlyName, bCertTimeIsValid ? "VALID" : "EXPIRED", EffectiveDate.ToString("o"), ExpirationDate.ToString("o"));
+					Program.LogVerbose("CERTIFICATE-Name:{0},Validity:{1},StartDate:{2},EndDate:{3}", CryptoAdapter.GetFriendlyNameFromCert(TestCert), bCertTimeIsValid ? "VALID" : "EXPIRED", EffectiveDate.ToString("o"), ExpirationDate.ToString("o"));
 				}
 
 				Store.Close();
@@ -384,7 +381,7 @@ namespace iPhonePackager
 					{
 						SelectedProvision = p.ProvisionName;
 						SelectedFile = Path.GetFileName(Provision);
-						SelectedCert = Cert.FriendlyName;
+						SelectedCert = CryptoAdapter.GetFriendlyNameFromCert(Cert);
 					}
 				}
 				Program.LogVerbose("PROVISION-File:{0},Name:{1},Validity:{2},StartDate:{3},EndDate:{4},Type:{5}", Path.GetFileName(Provision), p.ProvisionName, Validity, EffectiveDate.ToString(), ExpirationDate.ToString(), bDistribution ? "DISTRIBUTION" : "DEVELOPMENT");
@@ -431,10 +428,6 @@ namespace iPhonePackager
 
 						if (ValidInTimeCert != null)
 						{
-							int StartIndex = SourceCert.SubjectName.Name.IndexOf("CN=") + 3;
-							int EndIndex = SourceCert.SubjectName.Name.IndexOf(", ", StartIndex);
-							SourceCert.FriendlyName = SourceCert.SubjectName.Name.Substring(StartIndex, EndIndex - StartIndex);
-
 							// Found a cert in the valid time range, quit now!
 							Result = ValidInTimeCert;
 							break;
@@ -464,7 +457,7 @@ namespace iPhonePackager
 
 						bool bCertTimeIsValid = (EffectiveDate < Now) && (ExpirationDate > Now);
 
-						Program.LogVerbose ("  .. .. Installed certificate '{0}' is {1} (range '{2}' to '{3}')", TestCert.FriendlyName, bCertTimeIsValid ? "valid (choosing it)" : "EXPIRED", TestCert.GetEffectiveDateString (), TestCert.GetExpirationDateString ());
+						Program.LogVerbose ("  .. .. Installed certificate '{0}' is {1} (range '{2}' to '{3}')", CryptoAdapter.GetFriendlyNameFromCert(TestCert), bCertTimeIsValid ? "valid (choosing it)" : "EXPIRED", TestCert.GetEffectiveDateString (), TestCert.GetExpirationDateString ());
 						if (bCertTimeIsValid)
 						{
 							ValidInTimeCert = TestCert;
@@ -699,7 +692,7 @@ namespace iPhonePackager
 			}
 			else
 			{
-				Program.Log("... Found matching certificate '{0}' (valid from {1} to {2})", SigningCert.FriendlyName, SigningCert.GetEffectiveDateString(), SigningCert.GetExpirationDateString());
+				Program.Log("... Found matching certificate '{0}' (valid from {1} to {2})", CryptoAdapter.GetFriendlyNameFromCert(SigningCert), SigningCert.GetEffectiveDateString(), SigningCert.GetExpirationDateString());
 			}
 		}
 
diff --git a/Engine/Source/Programs/IOS/iPhonePackager/MachObjects.cs b/Engine/Source/Programs/IOS/iPhonePackager/MachObjects.cs
--- a/Engine/Source/Programs/IOS/iPhonePackager/MachObjects.cs
+++ b/Engine/Source/Programs/IOS/iPhonePackager/MachObjects.cs
@@ -1,4 +1,4 @@
-﻿// Copyright 1998-2019 Epic Games, Inc. All Rights Reserved.
+// Copyright 1998-2019 Epic Games, Inc. All Rights Reserved.
 
 
 using System;
@@ -1693,7 +1693,7 @@ namespace MachObjectHandling
 			}
 			if (string.IsNullOrEmpty(CertificateName))
 			{
-				CertificateName = SigningCert.FriendlyName;
+				CertificateName = CryptoAdapter.GetFriendlyNameFromCert(SigningCert);
 			}
 
 			// always use the new  requirements when initializing from cert
diff --git a/Engine/Source/Programs/IOS/iPhonePackager/MobileProvisionUtilities.cs b/Engine/Source/Programs/IOS/iPhonePackager/MobileProvisionUtilities.cs
--- a/Engine/Source/Programs/IOS/iPhonePackager/MobileProvisionUtilities.cs
+++ b/Engine/Source/Programs/IOS/iPhonePackager/MobileProvisionUtilities.cs
@@ -245,7 +245,7 @@ namespace iPhonePackager
 						bPassesHasMatchingCertCheck = (Cert != null);
 						if (bPassesHasMatchingCertCheck && Config.bCert)
 						{
-							bPassesHasMatchingCertCheck &= (Cert.FriendlyName == Config.Certificate);
+							bPassesHasMatchingCertCheck &= (CryptoAdapter.GetFriendlyNameFromCert(Cert) == Config.Certificate);
 						}
 					}
 					else
diff --git a/Engine/Source/Programs/IOS/iPhonePackager/Utilities.cs b/Engine/Source/Programs/IOS/iPhonePackager/Utilities.cs
--- a/Engine/Source/Programs/IOS/iPhonePackager/Utilities.cs
+++ b/Engine/Source/Programs/IOS/iPhonePackager/Utilities.cs
@@ -1,4 +1,4 @@
-﻿/**
+/**
  * Copyright 1998-2019 Epic Games, Inc. All Rights Reserved.
  */
 
@@ -918,21 +918,37 @@ namespace iPhonePackager
 			return CommonName;
 		}
 
-		/// <summary>
-		/// Merges a certificate and private key into a single combined certificate
-		/// </summary>
-		public static X509Certificate2 CombineKeyAndCert(string CertificateFilename, string KeyFilename)
+		public static string GetFriendlyNameFromCert(X509Certificate2 Cert)
+		{
+			if (Environment.OSVersion.Platform == PlatformID.Unix || Environment.OSVersion.Platform == PlatformID.MacOSX)
+			{
+				// From Mono 6, the FriendlyName field cannot be set on Unix so we always parse it out of the Subject.
+				return GetCommonNameFromCert(Cert);
+			}
+			else
+			{
+				// Make sure we have a useful friendly name
+				string FriendlyName = Cert.FriendlyName;
+				if ((FriendlyName == "") || (FriendlyName == null))
+				{
+					FriendlyName = GetCommonNameFromCert(Cert);
+				}
+
+				return FriendlyName;
+			}
+		}
+
+	/// <summary>
+	/// Merges a certificate and private key into a single combined certificate
+	/// </summary>
+	public static X509Certificate2 CombineKeyAndCert(string CertificateFilename, string KeyFilename)
 		{
 			// Load the certificate
 			string CertificatePassword = "";
 			X509Certificate2 Cert = new X509Certificate2(CertificateFilename, CertificatePassword, X509KeyStorageFlags.PersistKeySet | X509KeyStorageFlags.Exportable | X509KeyStorageFlags.MachineKeySet);
 
 			// Make sure we have a useful friendly name
-			string FriendlyName = Cert.FriendlyName;
-			if ((FriendlyName == "") || (FriendlyName == null))
-			{
-				FriendlyName = GetCommonNameFromCert(Cert);
-			}
+			string FriendlyName = GetFriendlyNameFromCert(Cert);
 
 			// Create a PKCS#12 store with both the certificate and the private key in it
 			Pkcs12Store Store = new Pkcs12StoreBuilder().Build();
diff --git a/Engine/Source/Programs/IOS/iPhonePackager/iPhonePackager.cs b/Engine/Source/Programs/IOS/iPhonePackager/iPhonePackager.cs
--- a/Engine/Source/Programs/IOS/iPhonePackager/iPhonePackager.cs
+++ b/Engine/Source/Programs/IOS/iPhonePackager/iPhonePackager.cs
@@ -1000,7 +1000,7 @@ namespace iPhonePackager
                             if (CodeSignatureBuilder.FindRequiredFiles(out Provision, out Cert, out bHasOverrideFile, out bNameMatch) && Cert != null)
                             {
                                 // print out the provision and cert name
-                                Program.LogVerbose("CERTIFICATE-{0},PROVISION-{1}", Cert.FriendlyName, Provision.FileName);
+                                Program.LogVerbose("CERTIFICATE-{0},PROVISION-{1}", CryptoAdapter.GetFriendlyNameFromCert(Cert), Provision.FileName);
                             }
                             else
                             {
