Fix issue with iOS remote build ssh command (Contributed by [REDACTED_PERSON])
PR #6504: Fixing ssh command for remote build (Contributed by [REDACTED_PERSON])

diff --git a/Engine/Source/Programs/IOS/iPhonePackager/SSHCommandHelper.cs b/Engine/Source/Programs/IOS/iPhonePackager/SSHCommandHelper.cs
--- a/Engine/Source/Programs/IOS/iPhonePackager/SSHCommandHelper.cs
+++ b/Engine/Source/Programs/IOS/iPhonePackager/SSHCommandHelper.cs
@@ -260,7 +260,7 @@ namespace iPhonePackager
 			// make simple rsync commandline to send a file
 			RsyncProcess.StartInfo.FileName = RSyncExePath;
 			RsyncProcess.StartInfo.Arguments = string.Format(
-                "-zrltgoDe \"ssh {0}\" --rsync-path=\"mkdir -p {1} && rsync\" --chmod=ug=rwX,o=rxX '{2}' \"{3}@{4}\":'{1}/{5}'",
+                "-zrltgoDe \"./ssh {0}\" --rsync-path=\"mkdir -p {1} && rsync\" --chmod=ug=rwX,o=rxX '{2}' \"{3}@{4}\":'{1}/{5}'",
 				SSHAuthentication,
 				RemoteDir,
 				ConvertPathToCygwin(LocalPath),
@@ -316,7 +316,7 @@ namespace iPhonePackager
 			// make simple rsync commandline to send a file
 			RsyncProcess.StartInfo.FileName = RSyncExePath;
 			RsyncProcess.StartInfo.Arguments = string.Format(
-                "-zrltgoDe \"ssh {0}\" \"{2}@{3}\":'{4}' \"{1}\"",
+                "-zrltgoDe \"./ssh {0}\" \"{2}@{3}\":'{4}' \"{1}\"",
 				SSHAuthentication,
 				ConvertPathToCygwin(LocalPath),
 				SSHUser,
diff --git a/Engine/Source/Programs/UnrealBuildTool/ToolChain/RemoteMac.cs b/Engine/Source/Programs/UnrealBuildTool/ToolChain/RemoteMac.cs
--- a/Engine/Source/Programs/UnrealBuildTool/ToolChain/RemoteMac.cs
+++ b/Engine/Source/Programs/UnrealBuildTool/ToolChain/RemoteMac.cs
@@ -58,7 +58,7 @@ namespace UnrealBuildTool
 		/// The authentication used for Rsync (for the -e rsync flag).
 		/// </summary>
 		[XmlConfigFile]
-		private string RsyncAuthentication = "ssh -i '${CYGWIN_SSH_PRIVATE_KEY}'";
+		private string RsyncAuthentication = "./ssh -i '${CYGWIN_SSH_PRIVATE_KEY}'";
 
 		/// <summary>
 		/// The authentication used for SSH (probably similar to RsyncAuthentication).
