[IOS] Add support for Apple Development and Distribution certificates.

diff --git a/Engine/Source/Programs/IOS/iPhonePackager/CodeSigning.cs b/Engine/Source/Programs/IOS/iPhonePackager/CodeSigning.cs
--- a/Engine/Source/Programs/IOS/iPhonePackager/CodeSigning.cs
+++ b/Engine/Source/Programs/IOS/iPhonePackager/CodeSigning.cs
@@ -1,4 +1,4 @@
-﻿/**
+/**
  * Copyright Epic Games, Inc. All Rights Reserved.
  */
 
@@ -220,39 +220,40 @@ namespace iPhonePackager
 		/// </summary>
 		public static void FindCertificates()
 		{
+			string[] ValidCertificatePrefixes = {"iPhone Developer", "iPhone Distribution", "Apple Development", "Apple Distribution"};
+			
+			X509Certificate2Collection FoundCerts = new X509Certificate2Collection();
+
 			if (Environment.OSVersion.Platform == PlatformID.Unix || Environment.OSVersion.Platform == PlatformID.MacOSX)
 			{
-				// run certtool y to get the currently installed certificates
-				CertToolData = "";
-				Process CertTool = new Process();
-				CertTool.StartInfo.FileName = "/usr/bin/security";
-				CertTool.StartInfo.UseShellExecute = false;
-				CertTool.StartInfo.Arguments = "find-certificate -a -c \"iPhone\" -p";
-				CertTool.StartInfo.RedirectStandardOutput = true;
-				CertTool.OutputDataReceived += new DataReceivedEventHandler(OutputReceivedCertToolProcessCall);
-				CertTool.Start();
-				CertTool.BeginOutputReadLine();
-				CertTool.WaitForExit();
-				if (CertTool.ExitCode == 0)
-				{
-					string header = "-----BEGIN CERTIFICATE-----\n";
-					string footer = "-----END CERTIFICATE-----";
-					int start = CertToolData.IndexOf(header);
-					while (start != -1)
+				foreach(string SearchPrefix in ValidCertificatePrefixes)
+				{ 
+					// run certtool y to get the currently installed certificates
+					CertToolData = "";
+					Process CertTool = new Process();
+					CertTool.StartInfo.FileName = "/usr/bin/security";
+					CertTool.StartInfo.UseShellExecute = false;
+					CertTool.StartInfo.Arguments = string.Format("find-certificate -a -c \"{0}\" -p", SearchPrefix);
+					CertTool.StartInfo.RedirectStandardOutput = true;
+					CertTool.OutputDataReceived += new DataReceivedEventHandler(OutputReceivedCertToolProcessCall);
+					CertTool.Start();
+					CertTool.BeginOutputReadLine();
+					CertTool.WaitForExit();
+					if (CertTool.ExitCode == 0)
 					{
-						start += header.Length;
-						int end = CertToolData.IndexOf(footer, start);
-						string base64 = CertToolData.Substring(start, (end - start));
-						byte[] certData = Convert.FromBase64String(base64);
-						X509Certificate2 cert = new X509Certificate2(certData);
-						DateTime EffectiveDate = cert.NotBefore.ToUniversalTime();
-						DateTime ExpirationDate = cert.NotAfter.ToUniversalTime();
-						DateTime Now = DateTime.UtcNow;
-
-						bool bCertTimeIsValid = (EffectiveDate < Now) && (ExpirationDate > Now);
-						Program.LogVerbose("CERTIFICATE-Name:{0},Validity:{1},StartDate:{2},EndDate:{3}", CryptoAdapter.GetFriendlyNameFromCert(cert), bCertTimeIsValid ? "VALID" : "EXPIRED", EffectiveDate.ToString("o"), ExpirationDate.ToString("o"));
-
-						start = CertToolData.IndexOf(header, start);
+						string header = "-----BEGIN CERTIFICATE-----\n";
+						string footer = "-----END CERTIFICATE-----";
+						int start = CertToolData.IndexOf(header);
+						while (start != -1)
+						{
+							start += header.Length;
+							int end = CertToolData.IndexOf(footer, start);
+							string base64 = CertToolData.Substring(start, (end - start));
+							byte[] certData = Convert.FromBase64String(base64);
+							X509Certificate2 cert = new X509Certificate2(certData);
+							FoundCerts.Add(cert);
+							start = CertToolData.IndexOf(header, start);
+						}
 					}
 				}
 			}
@@ -262,32 +263,22 @@ namespace iPhonePackager
 				X509Store Store = new X509Store();
 				Store.Open(OpenFlags.ReadOnly);
 
-				// Try finding all certificates that match either iPhone Developer or iPhone Distribution
-				X509Certificate2Collection FoundCerts = Store.Certificates.Find(X509FindType.FindBySubjectName, "iPhone Developer", false);
-
-				foreach (X509Certificate2 TestCert in FoundCerts)
+				foreach (string SearchPrefix in ValidCertificatePrefixes)
 				{
-					DateTime EffectiveDate = TestCert.NotBefore.ToUniversalTime();
-					DateTime ExpirationDate = TestCert.NotAfter.ToUniversalTime();
-					DateTime Now = DateTime.UtcNow;
-
-					bool bCertTimeIsValid = (EffectiveDate < Now) && (ExpirationDate > Now);
-					Program.LogVerbose("CERTIFICATE-Name:{0},Validity:{1},StartDate:{2},EndDate:{3}", CryptoAdapter.GetFriendlyNameFromCert(TestCert), bCertTimeIsValid ? "VALID" : "EXPIRED", EffectiveDate.ToString("o"), ExpirationDate.ToString("o"));
+					FoundCerts.AddRange(Store.Certificates.Find(X509FindType.FindBySubjectName, SearchPrefix, false));
 				}
 
-				FoundCerts = Store.Certificates.Find(X509FindType.FindBySubjectName, "iPhone Distribution", false);
-
-				foreach (X509Certificate2 TestCert in FoundCerts)
-				{
-					DateTime EffectiveDate = TestCert.NotBefore.ToUniversalTime();
-					DateTime ExpirationDate = TestCert.NotAfter.ToUniversalTime();
-					DateTime Now = DateTime.UtcNow;
+				Store.Close();
+			}
 
-					bool bCertTimeIsValid = (EffectiveDate < Now) && (ExpirationDate > Now);
-					Program.LogVerbose("CERTIFICATE-Name:{0},Validity:{1},StartDate:{2},EndDate:{3}", CryptoAdapter.GetFriendlyNameFromCert(TestCert), bCertTimeIsValid ? "VALID" : "EXPIRED", EffectiveDate.ToString("o"), ExpirationDate.ToString("o"));
-				}
+			foreach (X509Certificate2 TestCert in FoundCerts)
+			{
+				DateTime EffectiveDate = TestCert.NotBefore.ToUniversalTime();
+				DateTime ExpirationDate = TestCert.NotAfter.ToUniversalTime();
+				DateTime Now = DateTime.UtcNow;
 
-				Store.Close();
+				bool bCertTimeIsValid = (EffectiveDate < Now) && (ExpirationDate > Now);
+				Program.LogVerbose("CERTIFICATE-Name:{0},Validity:{1},StartDate:{2},EndDate:{3}", CryptoAdapter.GetFriendlyNameFromCert(TestCert), bCertTimeIsValid ? "VALID" : "EXPIRED", EffectiveDate.ToString("o"), ExpirationDate.ToString("o"));
 			}
 		}
 
