Fixed a number of cases where code was assuming Client targets are always named the same as the game target with "Client" appended.

diff --git a/Engine/Source/Programs/AutomationTool/IOS/IOSPlatform.Automation.cs b/Engine/Source/Programs/AutomationTool/IOS/IOSPlatform.Automation.cs
--- a/Engine/Source/Programs/AutomationTool/IOS/IOSPlatform.Automation.cs
+++ b/Engine/Source/Programs/AutomationTool/IOS/IOSPlatform.Automation.cs
@@ -194,15 +194,7 @@ public class IOSPlatform : Platform
 	{
 		string TargetName = Path.GetFileNameWithoutExtension(InExecutablePath).Split("-".ToCharArray())[0];
 		FileReference TargetReceiptFileName;
-		if (bIsUE4Game)
-		{
-			TargetReceiptFileName = TargetReceipt.GetDefaultPath(InEngineDir, "UE4Game", UnrealTargetPlatform.IOS, Config, "");
-		}
-		else
-		{
-			TargetReceiptFileName = TargetReceipt.GetDefaultPath(InProjectDirectory, TargetName, UnrealTargetPlatform.IOS, Config, "");
-		}
-
+		TargetReceiptFileName = TargetReceipt.GetDefaultPath(InProjectDirectory, TargetName, UnrealTargetPlatform.IOS, Config, "");
 		return TargetReceiptFileName;
 	}
 
@@ -213,16 +205,7 @@ public class IOSPlatform : Platform
 
 	public virtual bool DeployGeneratePList(FileReference ProjectFile, UnrealTargetConfiguration Config, DirectoryReference ProjectDirectory, bool bIsUE4Game, string GameName, bool bIsClient, string ProjectName, DirectoryReference InEngineDir, DirectoryReference AppDirectory, string InExecutablePath, out bool bSupportsPortrait, out bool bSupportsLandscape, out bool bSkipIcons)
 	{
-		string TargetName = Path.GetFileNameWithoutExtension(InExecutablePath).Split("-".ToCharArray())[0];
-		FileReference TargetReceiptFileName;
-		if (bIsUE4Game)
-		{
-			TargetReceiptFileName = TargetReceipt.GetDefaultPath(InEngineDir, "UE4Game", UnrealTargetPlatform.IOS, Config, "");
-		}
-		else
-		{
-			TargetReceiptFileName = TargetReceipt.GetDefaultPath(ProjectDirectory, TargetName, UnrealTargetPlatform.IOS, Config, "");
-		}
+		FileReference TargetReceiptFileName = GetTargetReceiptFileName(Config, InExecutablePath, InEngineDir, ProjectDirectory, bIsUE4Game);
 		return IOSExports.GeneratePList(ProjectFile, Config, ProjectDirectory, bIsUE4Game, GameName, bIsClient, ProjectName, InEngineDir, AppDirectory, TargetReceiptFileName, out bSupportsPortrait, out bSupportsLandscape, out bSkipIcons);
 	}
 
@@ -457,6 +440,7 @@ public class IOSPlatform : Platform
 				string IPPArguments = "RepackageFromStage \"" + (Params.IsCodeBasedProject ? Params.RawProjectPath.FullName : "Engine") + "\"";
 				IPPArguments += " -config " + TargetConfiguration.ToString();
 				IPPArguments += " -schemename " + SchemeName + " -schemeconfig \"" + SchemeConfiguration + "\"";
+				IPPArguments += " -targetname " + SC.StageExecutables[0];
 
 				if (TargetConfiguration == UnrealTargetConfiguration.Shipping)
 				{
@@ -867,7 +851,7 @@ public class IOSPlatform : Platform
 							(SC.IsCodeBasedProject ? SC.ProjectRoot : DirectoryReference.Combine(SC.LocalRoot, "Engine")),
 							!SC.IsCodeBasedProject,
 							(SC.IsCodeBasedProject ? SC.StageExecutables[0] : "UE4Game"),
-							Params.Client,
+							SC.IsCodeBasedProject ? false : Params.Client, // Code based projects will have Client in their executable name already
 							SC.ShortProjectName, DirectoryReference.Combine(SC.LocalRoot, "Engine"),
 							DirectoryReference.Combine((SC.IsCodeBasedProject ? SC.ProjectRoot : DirectoryReference.Combine(SC.LocalRoot, "Engine")), "Binaries", PlatformName, "Payload", (SC.IsCodeBasedProject ? SC.ShortProjectName : "UE4Game") + ".app"),
 							SC.StageExecutables[0],
diff --git a/Engine/Source/Programs/IOS/iPhonePackager/Config.cs b/Engine/Source/Programs/IOS/iPhonePackager/Config.cs
--- a/Engine/Source/Programs/IOS/iPhonePackager/Config.cs
+++ b/Engine/Source/Programs/IOS/iPhonePackager/Config.cs
@@ -1,4 +1,4 @@
-﻿/*
+/*
  * Copyright Epic Games, Inc. All Rights Reserved.
  */
 
@@ -168,7 +168,7 @@ namespace iPhonePackager
 		//@TODO: Deprecate this directory
 		public static string PayloadDirectory
 		{
-			get { return Path.GetFullPath(PCStagingRootDir + @"\Payload\" + Program.GameName + (Program.IsClient ? "Client" : "") + Program.Architecture + ".app"); }
+			get { return Path.GetFullPath(PCStagingRootDir + @"\Payload\" + GetTargetName() + Program.Architecture + ".app"); }
 		}
 
 		/// <summary>
@@ -196,11 +196,11 @@ namespace iPhonePackager
 			{
 				if (Program.GameConfiguration == "Development")
 				{
-					return SigningPrefix + Program.GameName + (Program.IsClient ? "Client" : "") + Program.Architecture + ".ipa";
+					return SigningPrefix + GetTargetName() + Program.Architecture + ".ipa";
 				}
 				else
 				{
-					return SigningPrefix + Program.GameName + (Program.IsClient ? "Client" : "") + "-" + Config.OSString + "-" + Program.GameConfiguration + Program.Architecture + ".ipa";
+					return SigningPrefix + GetTargetName() + "-" + Config.OSString + "-" + Program.GameConfiguration + Program.Architecture + ".ipa";
 				}
 			}
 		}
@@ -236,11 +236,11 @@ namespace iPhonePackager
 
 			if (Program.GameConfiguration == "Development")
 			{
-				Filename = Path.Combine(Config.BinariesDirectory, FilePrefix + Program.GameName + (Program.IsClient ? "Client" : "") + Program.Architecture + FileSuffix);
+				Filename = Path.Combine(Config.BinariesDirectory, FilePrefix + GetTargetName() + Program.Architecture + FileSuffix);
 			}
 			else
 			{
-				Filename = Path.Combine(Config.BinariesDirectory, FilePrefix + Program.GameName + (Program.IsClient ? "Client" : "") + "-" + Config.OSString + "-" + Program.GameConfiguration + Program.Architecture + FileSuffix);
+				Filename = Path.Combine(Config.BinariesDirectory, FilePrefix + GetTargetName() + "-" + Config.OSString + "-" + Program.GameConfiguration + Program.Architecture + FileSuffix);
 			}
 
 			return Filename;
@@ -254,19 +254,28 @@ namespace iPhonePackager
 			string Filename;
 
 			string BinariesDir = BinariesDirectory;
-			string GameName = Program.GameName;
+			string ExeName = Program.GameName;
 			if (!String.IsNullOrEmpty(ProjectFile))
 			{
 				BinariesDir = Path.Combine(Path.GetDirectoryName(ProjectFile), "Binaries", Config.OSString);
-				GameName = Path.GetFileNameWithoutExtension(ProjectFile);
+				ExeName = Path.GetFileNameWithoutExtension(ProjectFile);
 			}
+			if (!bIsCodeBasedProject)
+			{
+				ExeName = (Program.IsClient ? "UE4Client" : "UE4Game");
+			}
+			else if(Program.IsClient)
+			{
+				ExeName += "Client";
+			}
+
 			if (Program.GameConfiguration == "Development")
 			{
-				Filename = Path.Combine(BinariesDir, FilePrefix + (bIsCodeBasedProject ? GameName : "UE4Game") + (Program.IsClient ? "Client" : "") + Program.Architecture + FileSuffix);
+				Filename = Path.Combine(BinariesDir, FilePrefix + ExeName + Program.Architecture + FileSuffix);
 			}
 			else
 			{
-				Filename = Path.Combine(BinariesDir, FilePrefix + (bIsCodeBasedProject ? GameName : "UE4Game") + (Program.IsClient ? "Client" : "") + "-" + Config.OSString + "-" + Program.GameConfiguration + Program.Architecture + FileSuffix);
+				Filename = Path.Combine(BinariesDir, FilePrefix + ExeName + "-" + Config.OSString + "-" + Program.GameConfiguration + Program.Architecture + FileSuffix);
 			}
 
 			// ensure the directory exists
@@ -274,6 +283,15 @@ namespace iPhonePackager
 			return Filename;
 		}
 
+
+		/// <summary>
+		/// Returns target name, either the GameName + optional Client or specified by -targetname
+		/// </summary>
+		public static string GetTargetName()
+		{
+			return TargetName != "" ? TargetName : Program.GameName + (Program.IsClient ? "Client" : "");
+		}
+
 		/// <summary>
 		/// Returns the full path for the stub IPA, following the signing prefix resolution rules
 		/// </summary>
@@ -281,11 +299,11 @@ namespace iPhonePackager
 		{
 			if (Program.GameConfiguration == "Development")
 			{
-				return FileOperations.FindPrefixedFile(Config.BinariesDirectory, Program.GameName + (Program.IsClient ? "Client" : "") + Program.Architecture + FileSuffix);
+				return FileOperations.FindPrefixedFile(Config.BinariesDirectory, GetTargetName() + Program.Architecture + FileSuffix);
 			}
 			else
 			{
-				return FileOperations.FindPrefixedFile(Config.BinariesDirectory, Program.GameName + (Program.IsClient ? "Client" : "") + "-" + Config.OSString + "-" + Program.GameConfiguration + Program.Architecture + FileSuffix);
+				return FileOperations.FindPrefixedFile(Config.BinariesDirectory, GetTargetName() + "-" + Config.OSString + "-" + Program.GameConfiguration + Program.Architecture + FileSuffix);
 			}
 		}
 
@@ -355,6 +373,11 @@ namespace iPhonePackager
 		/// </summary>
 		public static string ProvisionFile;
 
+		/// <summary>
+		/// target name specificied by the command line, if not set GameName + optional Client will be used
+		/// </summary>
+		public static string TargetName = "";
+
 		/// <summary>
 		/// IOS Team ID to be used for automatic signing
 		/// </summary>
@@ -396,7 +419,7 @@ namespace iPhonePackager
 		/// </summary>
 		public static string AppDirectoryInZIP
 		{
-			get { return String.Format("Payload/{0}{1}.app", Program.GameName + (Program.IsClient ? "Client" : ""), Program.Architecture); }
+			get { return String.Format("Payload/{0}{1}.app", GetTargetName(), Program.Architecture); }
 		}
 
 		/// <summary>
@@ -484,7 +507,7 @@ namespace iPhonePackager
 			// Root directory on PC for staging files to copy to Mac
 			Config.PCStagingRootDir = String.Format(@"{0}-Deploy\{1}\{2}{3}\",
 				IntermediateDirectory,
-				Program.GameName + (Program.IsClient ? "Client" : ""),
+				GetTargetName(),
 				Program.GameConfiguration,
 				Program.Architecture);
 
diff --git a/Engine/Source/Programs/IOS/iPhonePackager/CookTime.cs b/Engine/Source/Programs/IOS/iPhonePackager/CookTime.cs
--- a/Engine/Source/Programs/IOS/iPhonePackager/CookTime.cs
+++ b/Engine/Source/Programs/IOS/iPhonePackager/CookTime.cs
@@ -1,4 +1,4 @@
-﻿/**
+/**
  * Copyright Epic Games, Inc. All Rights Reserved.
  */
 
@@ -149,23 +149,24 @@ namespace iPhonePackager
 
 		static public void CopySignedFiles()
 		{
+			string TargetName = Config.GetTargetName();
 			string NameDecoration;
 			if (Program.GameConfiguration == "Development")
 			{
-				NameDecoration = (Program.IsClient ? "Client" : "") + Program.Architecture;
+				NameDecoration = Program.Architecture;
 			}
 			else
 			{
-				NameDecoration = (Program.IsClient ? "Client" : "") + "-" + Config.OSString + "-" + Program.GameConfiguration + Program.Architecture;
+				NameDecoration = "-" + Config.OSString + "-" + Program.GameConfiguration + Program.Architecture;
 			}
 
 			// Copy and un-decorate the binary name
-			FileOperations.CopyFiles(Config.BinariesDirectory, Config.PayloadDirectory, "<PAYLOADDIR>", Program.GameName + NameDecoration, null);
-			FileOperations.RenameFile(Config.PayloadDirectory, Program.GameName + NameDecoration, Program.GameName);
+			FileOperations.CopyFiles(Config.BinariesDirectory, Config.PayloadDirectory, "<PAYLOADDIR>", TargetName + NameDecoration, null);
+			FileOperations.RenameFile(Config.PayloadDirectory, TargetName + NameDecoration, TargetName);
 
 			FileOperations.CopyNonEssentialFile(
-				Path.Combine(Config.BinariesDirectory, Program.GameName + NameDecoration + ".app.dSYM.zip"),
-				Path.Combine(Config.PCStagingRootDir, Program.GameName + NameDecoration + ".app.dSYM.zip.datecheck")
+				Path.Combine(Config.BinariesDirectory, TargetName + NameDecoration + ".app.dSYM.zip"),
+				Path.Combine(Config.PCStagingRootDir, TargetName + NameDecoration + ".app.dSYM.zip.datecheck")
 			);
 		}
 
@@ -222,7 +223,7 @@ namespace iPhonePackager
 				return;
 			}
 
-			string ZipWorkingDir = String.Format("Payload/{0}{1}.app/", Program.GameName + (Program.IsClient ? "Client" : ""), Program.Architecture);
+			string ZipWorkingDir = String.Format("Payload/{0}{1}.app/", Config.GetTargetName(), Program.Architecture);
 
 			FileOperations.ZipFileSystem FileSystem = new FileOperations.ZipFileSystem(Zip, ZipWorkingDir);
 
@@ -315,7 +316,7 @@ namespace iPhonePackager
 					string RelativeFilename = AbsoluteFilename.Substring(SourceDir.Length + 1).Replace('\\', '/');
 
 					string ZipAbsolutePath = String.Format("Payload/{0}{1}.app/{2}",
-						Program.GameName + (Program.IsClient ? "Client" : ""),
+						Config.GetTargetName(),
 						Program.Architecture,
 						RelativeFilename);
 
diff --git a/Engine/Source/Programs/IOS/iPhonePackager/iPhonePackager.cs b/Engine/Source/Programs/IOS/iPhonePackager/iPhonePackager.cs
--- a/Engine/Source/Programs/IOS/iPhonePackager/iPhonePackager.cs
+++ b/Engine/Source/Programs/IOS/iPhonePackager/iPhonePackager.cs
@@ -1,4 +1,4 @@
-﻿/**
+/**
  * Copyright Epic Games, Inc. All Rights Reserved.
  */
 
@@ -507,8 +507,20 @@ namespace iPhonePackager
                                 return false;
                             }
                         }
-                    }
-                    else
+						else if (Arg == "-targetname")
+						{
+							// make sure there's at least one more arg
+							if (Arguments.Length > ArgIndex + 1)
+							{
+								Config.TargetName = Arguments[++ArgIndex];
+							}
+							else
+							{
+								return false;
+							}
+						}
+					}
+					else
                     {
                         // RPC command
                         MainRPCCommand = Arguments[ArgIndex];
diff --git a/Engine/Source/Programs/UnrealBuildTool/Platform/IOS/IOSToolChain.cs b/Engine/Source/Programs/UnrealBuildTool/Platform/IOS/IOSToolChain.cs
--- a/Engine/Source/Programs/UnrealBuildTool/Platform/IOS/IOSToolChain.cs
+++ b/Engine/Source/Programs/UnrealBuildTool/Platform/IOS/IOSToolChain.cs
@@ -1986,7 +1986,7 @@ namespace UnrealBuildTool
 				{
 				    if (AppName == "UE4Game" || AppName == "UE4Client" || Target.ProjectFile == null || Target.ProjectFile.IsUnderDirectory(UnrealBuildTool.EngineDirectory))
 				    {
-					    GenerateProjectFiles(Target.ProjectFile, new string[] { "-platforms=" + (Target.Platform == UnrealTargetPlatform.IOS ? "IOS" : "TVOS"), "-NoIntellIsense", (Target.Platform == UnrealTargetPlatform.IOS ? "-iosdeployonly" : "-tvosdeployonly"), "-ignorejunk", (Target.bForDistribution ? "-distribution" : "-development"), "-bundleID=" + BundleID, "-includetemptargets" });
+					    GenerateProjectFiles(Target.ProjectFile, new string[] { "-platforms=" + (Target.Platform == UnrealTargetPlatform.IOS ? "IOS" : "TVOS"), "-NoIntellIsense", (Target.Platform == UnrealTargetPlatform.IOS ? "-iosdeployonly" : "-tvosdeployonly"), "-ignorejunk", (Target.bForDistribution ? "-distribution" : "-development"), "-bundleID=" + BundleID, "-includetemptargets", "-appname=" + AppName });
 				    }
 				    else
 				    {
diff --git a/Engine/Source/Programs/UnrealBuildTool/Platform/IOS/UEDeployIOS.cs b/Engine/Source/Programs/UnrealBuildTool/Platform/IOS/UEDeployIOS.cs
--- a/Engine/Source/Programs/UnrealBuildTool/Platform/IOS/UEDeployIOS.cs
+++ b/Engine/Source/Programs/UnrealBuildTool/Platform/IOS/UEDeployIOS.cs
@@ -209,7 +209,10 @@
 			string OldPListData = File.Exists(PListFile) ? File.ReadAllText(PListFile) : "";
 
 			string LaunchStoryboard = InEngineDir + "/Build/IOS/Resources/Interface/LaunchScreen.storyboard";
-			if (File.Exists(BuildDirectory + "/Resources/Interface/LaunchScreen.storyboard"))			{				LaunchStoryboard = BuildDirectory + "/Resources/Interface/LaunchScreen.storyboard";			}
+			if (File.Exists(BuildDirectory + "/Resources/Interface/LaunchScreen.storyboard"))
+			{
+				LaunchStoryboard = BuildDirectory + "/Resources/Interface/LaunchScreen.storyboard";
+			}
 
 			// get the settings from the ini file
 			// plist replacements
@@ -424,11 +427,9 @@
 			Text.AppendLine("\t<key>CFBundleDisplayName</key>");
 			Text.AppendLine(string.Format("\t<string>{0}</string>", EncodeBundleName(BundleDisplayName, ProjectName)));
 			Text.AppendLine("\t<key>CFBundleExecutable</key>");
-			string BundleExecutable = bIsUE4Game ? "UE4Game" : GameName;
-			if (bIsClient)
-			{
-				BundleExecutable += "Client";
-			}
+			string BundleExecutable = bIsUE4Game ?
+				(bIsClient ? "UE4Client" : "UE4Game") :
+				(bIsClient ? GameName + "Client" : GameName);
 			Text.AppendLine(string.Format("\t<string>{0}</string>", BundleExecutable));
 			Text.AppendLine("\t<key>CFBundleIdentifier</key>");
 			Text.AppendLine(string.Format("\t<string>{0}</string>", BundleIdentifier.Replace("[PROJECT_NAME]", ProjectName).Replace("_", "")));
diff --git a/Engine/Source/Programs/UnrealBuildTool/Platform/TVOS/UEDeployTVOS.cs b/Engine/Source/Programs/UnrealBuildTool/Platform/TVOS/UEDeployTVOS.cs
--- a/Engine/Source/Programs/UnrealBuildTool/Platform/TVOS/UEDeployTVOS.cs
+++ b/Engine/Source/Programs/UnrealBuildTool/Platform/TVOS/UEDeployTVOS.cs
@@ -116,11 +116,9 @@ namespace UnrealBuildTool
 			Text.AppendLine("\t<key>CFBundleDisplayName</key>");
 			Text.AppendLine(string.Format("\t<string>{0}</string>", EncodeBundleName(BundleDisplayName, ProjectName)));
 			Text.AppendLine("\t<key>CFBundleExecutable</key>");
-			string BundleExecutable = bIsUE4Game ? "UE4Game" : GameName;
-			if (bIsClient)
-			{
-				BundleExecutable += "Client";
-			}
+			string BundleExecutable = bIsUE4Game ?
+				(bIsClient ? "UE4Client" : "UE4Game") :
+				(bIsClient ? GameName + "Client" : GameName);
 			Text.AppendLine(string.Format("\t<string>{0}</string>", BundleExecutable));
 			Text.AppendLine("\t<key>CFBundleIdentifier</key>");
 			Text.AppendLine(string.Format("\t<string>{0}</string>", BundleIdentifier.Replace("[PROJECT_NAME]", ProjectName).Replace("_","")));
diff --git a/Engine/Source/Programs/UnrealBuildTool/ProjectFiles/Xcode/XcodeProject.cs b/Engine/Source/Programs/UnrealBuildTool/ProjectFiles/Xcode/XcodeProject.cs
--- a/Engine/Source/Programs/UnrealBuildTool/ProjectFiles/Xcode/XcodeProject.cs
+++ b/Engine/Source/Programs/UnrealBuildTool/ProjectFiles/Xcode/XcodeProject.cs
@@ -135,12 +135,13 @@ namespace UnrealBuildTool
 		/// <param name="InOnlyGameProject"></param>
 		/// <param name="IsForDistribution">True for distribution builds</param>
 		/// <param name="BundleID">Override option for bundle identifier</param>
-		public XcodeProjectFile(FileReference InitFilePath, FileReference InOnlyGameProject, bool IsForDistribution, string BundleID)
+		public XcodeProjectFile(FileReference InitFilePath, FileReference InOnlyGameProject, bool IsForDistribution, string BundleID, string InAppName)
 			: base(InitFilePath)
 		{
 			//OnlyGameProject = InOnlyGameProject;
 			bForDistribution = IsForDistribution;
 			BundleIdentifier = BundleID;
+			AppName = InAppName;
 		}
 
 		public override string ToString()
@@ -158,6 +159,11 @@ namespace UnrealBuildTool
 		/// </summary>
 		string BundleIdentifier = "";
 
+		/// <summary>
+		/// Override AppName
+		/// </summary>
+		string AppName = "";
+
 		/// <summary>
 		/// Gets Xcode file category based on its extension
 		/// </summary>
@@ -1232,6 +1238,7 @@ namespace UnrealBuildTool
 						DirectoryReference ProjectPath = GameDir;
 						DirectoryReference EngineDir = DirectoryReference.Combine(new DirectoryReference(UE4Dir), "Engine");
 						string GameName = Config.BuildTarget;
+						bool bIsClient = false;
 						if (ProjectPath == null)
 						{
 							ProjectPath = EngineDir;
@@ -1240,6 +1247,7 @@ namespace UnrealBuildTool
 						{
 							ProjectPath = EngineDir;
 							GameName = "UE4Game";
+							bIsClient = (AppName == "UE4Client");
 						}
 
 						if (bCreateMacInfoPlist)
@@ -1265,12 +1273,12 @@ namespace UnrealBuildTool
 							TargetReceipt.TryRead(ReceiptFilename, out Receipt);
 							VersionNumber SdkVersion = UEDeployIOS.GetSdkVersion(Receipt);
 							bool bBuildAsFramework = UEDeployIOS.GetCompileAsDll(Receipt);
-							UEDeployIOS.GenerateIOSPList(ProjectFile, Config.BuildConfig, ProjectPath.FullName, bIsUE4Game, GameName, false, Config.BuildTarget, EngineDir.FullName, ProjectPath + "/Binaries/IOS/Payload", SdkVersion, null, BundleIdentifier, bBuildAsFramework, out bSupportPortrait, out bSupportLandscape, out bSkipIcons);
+							UEDeployIOS.GenerateIOSPList(ProjectFile, Config.BuildConfig, ProjectPath.FullName, bIsUE4Game, GameName, bIsClient, Config.BuildTarget, EngineDir.FullName, ProjectPath + "/Binaries/IOS/Payload", SdkVersion, null, BundleIdentifier, bBuildAsFramework, out bSupportPortrait, out bSupportLandscape, out bSkipIcons);
 						}
 						if (bCreateTVOSInfoPlist)
 						{
 							Directory.CreateDirectory(Path.GetDirectoryName(TVOSInfoPlistPath));
-							UEDeployTVOS.GenerateTVOSPList(ProjectPath.FullName, bIsUE4Game, GameName, false, Config.BuildTarget, EngineDir.FullName, ProjectPath + "/Binaries/TVOS/Payload", null, BundleIdentifier);
+							UEDeployTVOS.GenerateTVOSPList(ProjectPath.FullName, bIsUE4Game, GameName, bIsClient, Config.BuildTarget, EngineDir.FullName, ProjectPath + "/Binaries/TVOS/Payload", null, BundleIdentifier);
 						}
 					}
 				}
diff --git a/Engine/Source/Programs/UnrealBuildTool/ProjectFiles/Xcode/XcodeProjectFileGenerator.cs b/Engine/Source/Programs/UnrealBuildTool/ProjectFiles/Xcode/XcodeProjectFileGenerator.cs
--- a/Engine/Source/Programs/UnrealBuildTool/ProjectFiles/Xcode/XcodeProjectFileGenerator.cs
+++ b/Engine/Source/Programs/UnrealBuildTool/ProjectFiles/Xcode/XcodeProjectFileGenerator.cs
@@ -37,6 +37,11 @@ namespace UnrealBuildTool
 		/// </summary>
 		string BundleIdentifier = "";
 
+		/// <summary>
+		/// Override AppName
+		/// </summary>
+		string AppName = "";
+
 		public XcodeProjectFileGenerator(FileReference InOnlyGameProject, CommandLineArguments CommandLine)
 			: base(InOnlyGameProject)
 		{
@@ -48,6 +53,11 @@ namespace UnrealBuildTool
 			{
 				BundleIdentifier = CommandLine.GetString("-bundleID=");
 			}
+
+			if (CommandLine.HasValue("-appname="))
+			{
+				AppName = CommandLine.GetString("-appname=");
+			}
 		}
 
 		/// <summary>
@@ -108,7 +118,7 @@ namespace UnrealBuildTool
 		/// <returns>The newly allocated project file object</returns>
 		protected override ProjectFile AllocateProjectFile(FileReference InitFilePath)
 		{
-			return new XcodeProjectFile(InitFilePath, OnlyGameProject, bForDistribution, BundleIdentifier);
+			return new XcodeProjectFile(InitFilePath, OnlyGameProject, bForDistribution, BundleIdentifier, AppName);
 		}
 
 		/// ProjectFileGenerator interface
