Fixed iOS failure to match find profile when packaging for distribution with an ad-hoc provisioning profile
Fixed issue with fix for UE-77351 with Shipping builds
Fixed iOS automation failure enumerating folders where there user does not have permission the read the folder

diff --git a/Engine/Source/Programs/AutomationTool/IOS/IOSPlatform.Automation.cs b/Engine/Source/Programs/AutomationTool/IOS/IOSPlatform.Automation.cs
--- a/Engine/Source/Programs/AutomationTool/IOS/IOSPlatform.Automation.cs
+++ b/Engine/Source/Programs/AutomationTool/IOS/IOSPlatform.Automation.cs
@@ -440,7 +440,9 @@ public class IOSPlatform : Platform
 				string IPPArguments = "RepackageFromStage \"" + (Params.IsCodeBasedProject ? Params.RawProjectPath.FullName : "Engine") + "\"";
 				IPPArguments += " -config " + TargetConfiguration.ToString();
 				IPPArguments += " -schemename " + SchemeName + " -schemeconfig \"" + SchemeConfiguration + "\"";
-				IPPArguments += " -targetname " + SC.StageExecutables[0];
+
+				// targetname will be eg FooClient for a Client Shipping build.
+				IPPArguments += " -targetname " + SC.StageExecutables[0].Split("-".ToCharArray())[0];
 
 				if (TargetConfiguration == UnrealTargetConfiguration.Shipping)
 				{
diff --git a/Engine/Source/Programs/DotNETCommon/DotNETUtilities/DirectoryUtils.cs b/Engine/Source/Programs/DotNETCommon/DotNETUtilities/DirectoryUtils.cs
--- a/Engine/Source/Programs/DotNETCommon/DotNETUtilities/DirectoryUtils.cs
+++ b/Engine/Source/Programs/DotNETCommon/DotNETUtilities/DirectoryUtils.cs
@@ -37,13 +37,22 @@ namespace Tools.DotNETCommon
 			else
 			{
 				ParentInfo = FindCorrectCase(ParentInfo);
-				foreach (DirectoryInfo ChildInfo in ParentInfo.EnumerateDirectories())
+				try
 				{
-					if (String.Equals(ChildInfo.Name, Info.Name, DirectoryReference.Comparison))
+					foreach (DirectoryInfo ChildInfo in ParentInfo.EnumerateDirectories())
 					{
-						return ChildInfo;
+						if (String.Equals(ChildInfo.Name, Info.Name, DirectoryReference.Comparison))
+						{
+							return ChildInfo;
+						}
 					}
 				}
+				catch (Exception)
+				{
+					// System.Security.SecurityException is expected here if we try to enumerate a directory
+					// we don't have permission to read, such as another user's mounted volume.
+				}
+
 				return new DirectoryInfo(Path.Combine(ParentInfo.FullName, Info.Name));
 			}
 		}
diff --git a/Engine/Source/Programs/IOS/iPhonePackager/MobileProvisionUtilities.cs b/Engine/Source/Programs/IOS/iPhonePackager/MobileProvisionUtilities.cs
--- a/Engine/Source/Programs/IOS/iPhonePackager/MobileProvisionUtilities.cs
+++ b/Engine/Source/Programs/IOS/iPhonePackager/MobileProvisionUtilities.cs
@@ -207,8 +207,9 @@ namespace iPhonePackager
 
                     if (Config.bForDistribution)
                     {
-                        // check to see if this is a distribution provision
-                        bool bDistroProv = (TestProvision.ProvisionedDeviceIDs.Count == 0) && !TestProvision.bDebug;
+						// Check to see if this is a distribution provision. get-task-allow must be false for distro profiles.
+						// TestProvision.ProvisionedDeviceIDs.Count==0 is not a valid check as ad-hoc distro profiles do list devices.
+						bool bDistroProv = !TestProvision.bDebug;
                         if (!bDistroProv)
                         {
                             Program.LogVerbose("  .. Failed distribution check (mode={0}, get-task-allow={1}, #devices={2})", Config.bForDistribution, TestProvision.bDebug, TestProvision.ProvisionedDeviceIDs.Count);
