Remove reference to epicgames.common to avoid build errors when building this net fw project. A more proper solution is described in UE-103939 were we should migrate this tool to net core.

diff --git a/Engine/Source/Programs/IOS/iPhonePackager/MobileProvisionUtilities.cs b/Engine/Source/Programs/IOS/iPhonePackager/MobileProvisionUtilities.cs
--- a/Engine/Source/Programs/IOS/iPhonePackager/MobileProvisionUtilities.cs
+++ b/Engine/Source/Programs/IOS/iPhonePackager/MobileProvisionUtilities.cs
@@ -11,7 +11,6 @@ using System.IO;
 using System.Diagnostics;
 using System.Xml;
 using System.Globalization;
-using EpicGames.Core;
 using System.Linq;
 
 namespace iPhonePackager
@@ -55,7 +54,7 @@ namespace iPhonePackager
                 {
                     foreach (string Provision in Directory.EnumerateFiles(ProjectFileBuildIOSPath, "*.mobileprovision", SearchOption.AllDirectories))
                     {
-                        Log.TraceInformation(Provision);
+                        Program.Log(Provision);
                         string TargetFile = Config.ProvisionDirectory + Path.GetFileName(Provision);
                         if (!File.Exists(TargetFile) || File.GetLastWriteTime(TargetFile) < File.GetLastWriteTime(Provision))
                         {
@@ -86,7 +85,7 @@ namespace iPhonePackager
                 {
                     foreach (string Provision in Directory.EnumerateFiles(ProvisionDirectory, "*.mobileprovision", SearchOption.AllDirectories))
                     {
-                        Log.TraceInformation(Provision);
+                        Program.Log(Provision);
                         string TargetFile = Config.ProvisionDirectory + Path.GetFileName(Provision);
                         if (!File.Exists(TargetFile) || File.GetLastWriteTime(TargetFile) < File.GetLastWriteTime(Provision))
                         {
diff --git a/Engine/Source/Programs/IOS/iPhonePackager/iPhonePackager.cs b/Engine/Source/Programs/IOS/iPhonePackager/iPhonePackager.cs
--- a/Engine/Source/Programs/IOS/iPhonePackager/iPhonePackager.cs
+++ b/Engine/Source/Programs/IOS/iPhonePackager/iPhonePackager.cs
@@ -16,7 +16,6 @@ using Ionic.Zlib;
 using System.ComponentModel;
 using System.Security.Cryptography.X509Certificates;
 using Microsoft.Win32;
-using EpicGames.Core;
 
 namespace iPhonePackager
 {
diff --git a/Engine/Source/Programs/IOS/iPhonePackager/iPhonePackager.csproj b/Engine/Source/Programs/IOS/iPhonePackager/iPhonePackager.csproj
--- a/Engine/Source/Programs/IOS/iPhonePackager/iPhonePackager.csproj
+++ b/Engine/Source/Programs/IOS/iPhonePackager/iPhonePackager.csproj
@@ -222,10 +222,6 @@
     <None Include="Resources\GreyCheck.png" />
   </ItemGroup>
   <ItemGroup>
-    <ProjectReference Include="..\..\Shared\EpicGames.Core\EpicGames.Core.csproj">
-      <Project>{5d7d66e8-8c76-4af9-b3ec-2ef03421d730}</Project>
-      <Name>DotNETUtilities</Name>
-    </ProjectReference>
     <ProjectReference Include="..\DeploymentInterface\DeploymentInterface.csproj">
       <Project>{B7C89A17-1CB7-42D7-B879-43297538C193}</Project>
       <Name>DeploymentInterface</Name>
