iPhone Packager Refactor #rb jack.porter #jira UE-103878

diff --git a/Engine/Source/Programs/IOS/iPhonePackager/DeployTime.cs b/Engine/Source/Programs/IOS/iPhonePackager/DeployTime.cs
--- a/Engine/Source/Programs/IOS/iPhonePackager/DeployTime.cs
+++ b/Engine/Source/Programs/IOS/iPhonePackager/DeployTime.cs
@@ -19,9 +19,27 @@ using System.Runtime.Remoting;
 
 namespace iPhonePackager
 {
+	[Serializable]
+	public struct ConnectedDeviceInformation
+	{
+		public string DeviceName;
+		public string UDID;
+		public string DeviceType;
+		public string DeviceInterface;
+
+
+	public ConnectedDeviceInformation(string DeviceName, string UDID, string DeviceType, string DeviceInterface)
+		{
+			this.DeviceName = DeviceName;
+			this.UDID = UDID;
+			this.DeviceType = DeviceType;
+			this.DeviceInterface = DeviceInterface;
+		}
+	}
+
 	// Implementation of DeploymentServer -> Application interface
 	[Serializable]
-	class DeployTimeReporter : DeployTimeReportingInterface
+	class DeployTimeReporter
 	{
 		public void Log(string Line)
 		{
@@ -51,10 +69,6 @@ namespace iPhonePackager
 
 	class DeploymentHelper
 	{
-		static DeploymentInterface DeployTimeInstance;
-
-		public static Process DeploymentServerProcess = null;
-
 		public static void InstallIPAOnConnectedDevices(string IPAPath)
 		{
 			// Read the mobile provision to check for issues
@@ -71,11 +85,12 @@ namespace iPhonePackager
 			}
 			Zip.Close();
 
+
 			if (Provision != null)
 			{
-				var DeviceList = DeploymentHelper.Get().EnumerateConnectedDevices();
+				var DevicesList = GetAllConnectedDevices();
 
-				foreach (var DeviceInfo in DeviceList)
+				foreach (var DeviceInfo in DevicesList)
 				{
 					string UDID = DeviceInfo.UDID;
 					string DeviceName = DeviceInfo.DeviceName;
@@ -95,8 +110,7 @@ namespace iPhonePackager
 					}
 				}
 			}
-
-			DeploymentHelper.Get().InstallIPAOnDevice(IPAPath);
+			InstallIPAOnDevices(IPAPath);
 		}
 
 		public static bool ExecuteDeployCommand(string Command, string GamePath, string RPCCommand)
@@ -113,13 +127,13 @@ namespace iPhonePackager
 
 						if (Config.FilesForBackup.Count > 0)
 						{
-							if (!DeploymentHelper.Get().BackupFiles(ApplicationIdentifier, Config.FilesForBackup.ToArray()))
+							if (!BackupFiles(ApplicationIdentifier, Config.FilesForBackup.ToArray()))
 							{
 								Program.Error("Failed to transfer manifest file from device to PC");
 								Program.ReturnCode = (int)ErrorCodes.Error_DeviceBackupFailed;
 							}
 						}
-						else if (!DeploymentHelper.Get().BackupDocumentsDirectory(ApplicationIdentifier, Config.GetRootBackedUpDocumentsDirectory()))
+						else if (!DeploymentHelper.BackupDocumentsDirectory(ApplicationIdentifier, Config.GetRootBackedUpDocumentsDirectory()))
 						{
 							Program.Error("Failed to transfer documents directory from device to PC");
 							Program.ReturnCode = (int)ErrorCodes.Error_DeviceBackupFailed;
@@ -134,7 +148,7 @@ namespace iPhonePackager
 							ApplicationIdentifier = Utilities.GetStringFromPList("CFBundleIdentifier");
 						}
 
-						if (!DeploymentHelper.Get().UninstallIPAOnDevice(ApplicationIdentifier))
+						if (!UninstallIPAOnDevices(ApplicationIdentifier))
 						{
 							Program.Error("Failed to uninstall IPA on device");
 							Program.ReturnCode = (int)ErrorCodes.Error_AppUninstallFailed;
@@ -208,8 +222,7 @@ namespace iPhonePackager
 								ApplicationIdentifier = Utilities.GetStringFromPList("CFBundleIdentifier");
 							}
 
-							DeploymentHelper.Get().DeviceId = Config.DeviceId;
-							if (!DeploymentHelper.Get().InstallFilesOnDevice(ApplicationIdentifier, Config.DeltaManifest))
+							if (!DeploymentHelper.PushFiles(ApplicationIdentifier, Config.DeltaManifest))
 							{
 								Program.Error("Failed to install Files on device");
 								Program.ReturnCode = (int)ErrorCodes.Error_FilesInstallFailed;
@@ -217,8 +230,7 @@ namespace iPhonePackager
 						}
 						else if (File.Exists(IPAPath))
 						{
-							DeploymentHelper.Get().DeviceId = Config.DeviceId;
-							if (!DeploymentHelper.Get().InstallIPAOnDevice(IPAPath))
+							if (!InstallIPAOnDevices(IPAPath))
 							{
 								Program.Error("Failed to install IPA on device");
 								Program.ReturnCode = (int)ErrorCodes.Error_AppInstallFailed;
@@ -240,62 +252,283 @@ namespace iPhonePackager
 		}
 
 		static DeployTimeReporter Reporter = new DeployTimeReporter();
-
-		public static DeploymentInterface Get()
+		static string ExecuteLibimobileProcess(string ProcessName, string ProcessArguments = "")
 		{
-			if (DeployTimeInstance == null)
+			string ExePath = Directory.GetCurrentDirectory() + "/../../../Extras/ThirdPartyNotUE/libimobiledevice/";
+
+			if (Environment.OSVersion.Platform == PlatformID.Win32NT)
+			{
+				ExePath += "x64/";
+				ProcessName += ".exe";
+			}
+			else if (Environment.OSVersion.Platform == PlatformID.MacOSX)
+			{
+				ExePath += "Mac/";
+			}
+			else
+			{
+				Program.LogVerbose("Unsupported Platform.");
+				return "";
+			}
+
+			ProcessStartInfo StartInfo = new ProcessStartInfo(ExePath + ProcessName , ProcessArguments);
+			StartInfo.UseShellExecute = false;
+			StartInfo.RedirectStandardOutput = true;
+			StartInfo.RedirectStandardError = true;
+			StartInfo.CreateNoWindow = true;
+			StartInfo.StandardOutputEncoding = Encoding.UTF8;
+
+			string FullOutput = "";
+			string ErrorOutput = "";
+			using (Process LocalProcess = Process.Start(StartInfo))
 			{
-				if (DeploymentServerProcess == null)
+				StreamReader OutputReader = LocalProcess.StandardOutput;
+				FullOutput = OutputReader.ReadToEnd().Trim();
+
+				StreamReader ErrorReader = LocalProcess.StandardError;
+				ErrorOutput = ErrorReader.ReadToEnd().Trim();
+
+				if (FullOutput.Length > 0)
 				{
-					DeploymentServerProcess = CreateDeploymentServerProcess();
+					Program.LogVerbose(FullOutput);
 				}
-				DeployTimeInstance = (DeploymentInterface)Activator.GetObject(
-					typeof(DeploymentInterface),
-					@"ipc://iPhonePackager/DeploymentServer_PID" + Process.GetCurrentProcess().Id.ToString());
+
+				if (ErrorOutput.Length > 0)
+				{
+					Program.LogVerbose(ErrorOutput);
+				}
+
+
+				LocalProcess.WaitForExit();
 			}
 
-			if (DeployTimeInstance == null)
+			if (ErrorOutput.Length > 0)
 			{
-				Program.Error("Failed to connect to deployment server");
-				throw new Exception("Failed to connect to deployment server");
+				if (FullOutput.Length > 0)
+				{
+					FullOutput += Environment.NewLine;
+				}
+				FullOutput += ErrorOutput;
 			}
+			return FullOutput;
+		}
+		public static  bool PushFiles(string BundleIdentifier, string ManifestFile)
+		{
+			string FullOutput = "";
+			var DeviceList = GetAllConnectedDevices();
 
-			DeployTimeInstance.SetReportingInterface(Reporter);
+			foreach (var Device in DeviceList)
+			{
+				string IdeviceFSArgs = "";
+				if (Device.DeviceInterface == "Network")
+				{
+					IdeviceFSArgs = "-n";
+					continue;
+				}
 
-			return DeployTimeInstance;
-		}
+				string[] FileList = ManifestFile.Split('\n');
+				string AllCommandsToPush = "";
+				foreach (string Filename in FileList)
+				{
+					if (!string.IsNullOrEmpty(Filename) && !string.IsNullOrWhiteSpace(Filename))
+					{
+						string Trimmed = Filename.Trim();
+						string BaseFolder = Path.GetDirectoryName(ManifestFile);
+						string SourceFilename = BaseFolder + "/" + Trimmed;
+						string DestFilename = "/Library/Caches/" + Trimmed.Replace("cookeddata/", "");
+						DestFilename = DestFilename.Replace('\\', '/');
+						DestFilename = "\"" + DestFilename + "\"";
+						string CommandToPush = "push -p \"" + SourceFilename + "\" " + DestFilename + "\n";
+						AllCommandsToPush += CommandToPush;
+					}
+				}
+				System.IO.File.WriteAllText(Directory.GetCurrentDirectory() + "/CommandsToPush.txt", AllCommandsToPush);
 
-		static Process CreateDeploymentServerProcess()
-		{
-			Process NewProcess = new Process();
+				IdeviceFSArgs = IdeviceFSArgs + " -u " + Device.UDID + " -b " + BundleIdentifier + " -x \"" + Directory.GetCurrentDirectory() + "/CommandsToPush.txt" + "\"";
+
+				FullOutput = ExecuteLibimobileProcess("idevicefs", IdeviceFSArgs);
+				File.Delete(Directory.GetCurrentDirectory() + "/CommandsToPush.txt");
 
-			NewProcess.StartInfo.WorkingDirectory = Path.GetFullPath(".");
+				if (FullOutput != "0")
+				{
+					return false;
+				}
+			}
+			return true;
+		}
 
-			if (Environment.OSVersion.Platform == PlatformID.MacOSX || Environment.OSVersion.Platform == PlatformID.Unix)
+		static bool BackupFiles(string BundleIdentifier, string[] DestinationFiles)
+		{
+			string FullOutput = "";
+			var DeviceList = GetAllConnectedDevices();
+			foreach (var Device in DeviceList)
 			{
-				NewProcess.StartInfo.WorkingDirectory.TrimEnd('/');
-				NewProcess.StartInfo.FileName = "../../../Build/BatchFiles/Mac/RunMono.sh";
-				NewProcess.StartInfo.Arguments = "\"" + NewProcess.StartInfo.WorkingDirectory + "/DeploymentServer.exe\" -iphonepackager " + Process.GetCurrentProcess().Id.ToString();
+				if (Device.DeviceInterface == "Network")
+				{
+					continue;
+				}
+
+				string AllCommands = "";
+				foreach (var DestinationFile in DestinationFiles)
+				{
+					string Command = "pull " + DestinationFile;
+					AllCommands += Command;
+				}
+				System.IO.File.WriteAllText(Directory.GetCurrentDirectory() + "/CommandsToPull.txt", AllCommands);
+				string arguments = "\"" + "-u " + Device.UDID + " -b " + BundleIdentifier + " -x " + Directory.GetCurrentDirectory() + "/CommandsToPull.txt" + "\"";
+				FullOutput = ExecuteLibimobileProcess("idevicefs", arguments);
 			}
-			else
+			File.Delete(Directory.GetCurrentDirectory() + "/CommandsToPull.txt");
+			if (FullOutput != "0")
 			{
-				NewProcess.StartInfo.WorkingDirectory.TrimEnd('\\');
-				NewProcess.StartInfo.FileName = NewProcess.StartInfo.WorkingDirectory + "\\DeploymentServer.exe";
-				NewProcess.StartInfo.Arguments = "-iphonepackager " + Process.GetCurrentProcess().Id.ToString();
+				return false;
 			}
-			NewProcess.StartInfo.UseShellExecute = false;
+			return true;
+		}
 
-			try
+		public static bool BackupDocumentsDirectory(string BundleIdentifier, string DestinationDocumentsDirectory)
+		{
+			string FullOutput = "";
+			var DeviceList = GetAllConnectedDevices();
+			foreach (var Device in DeviceList)
 			{
-				NewProcess.Start();
-				System.Threading.Thread.Sleep(1000);
-			}
-			catch (System.Exception ex)
+				if (Device.DeviceInterface == "Network")
+				{
+					continue;
+				}
+
+				// Destination folder
+				string TargetFolder = Path.Combine(DestinationDocumentsDirectory, Device.DeviceName);
+
+				// Source folder
+				string SourceFolder = "/Documents/ ";
+
+				string Command = " pull " + SourceFolder + "\"" +  TargetFolder + "\"";
+
+				FullOutput = ExecuteLibimobileProcess("idevicefs", "-u " + Device.UDID + " -b " + BundleIdentifier + Command);
+				if (FullOutput != "0")
+				{
+					Program.LogVerbose(FullOutput);
+				}
+				SourceFolder = "/Library/Caches/";
+				Command = " pull " + SourceFolder;
+				FullOutput = ExecuteLibimobileProcess("idevicefs", "-u " + Device.UDID + " -b " + BundleIdentifier + Command);
+				if (FullOutput != "0")
+				{
+					Program.LogVerbose(FullOutput);
+				}
+			}			
+			return true;
+		}
+
+		static bool ActionOnIPAOnDevices(string args, string IPAPath)
+		{
+			List<ConnectedDeviceInformation> DevicesToReturn = new List<ConnectedDeviceInformation>();
+			string FullOutput = ExecuteLibimobileProcess("idevice_id");
+			string LibimobileDeviceArguments = "";
+
+			var ConnectedDevicesUDIDs = FullOutput.Split(new string[] { Environment.NewLine }, StringSplitOptions.None);
+
+			foreach (var ConnnectedUDID in ConnectedDevicesUDIDs)
 			{
-				Program.Error("Failed to create deployment server process ({0})", ex.Message);
+				string CurrentUDID = ConnnectedUDID;
+				if (ConnnectedUDID.Contains("(USB)"))
+				{
+					CurrentUDID = ConnnectedUDID.Split(' ').First();
+					LibimobileDeviceArguments = "-u " + CurrentUDID + args + "\"" + IPAPath + "\"";
+					string OutputInfo = ExecuteLibimobileProcess("ideviceinstaller", LibimobileDeviceArguments);
+					if (OutputInfo.Contains("DONE"))
+					{
+						return true;
+					}
+					else
+					{
+						Program.LogVerbose(OutputInfo);
+						return false;
+					}
+				}
+				else if (ConnnectedUDID.Contains("(Network)"))
+				{
+					// Network install doesn't seem to be working on Windows. USB Interface only for now.
+					//	CurrentUDID = ConnnectedUDID.Split(' ').First();
+					//	LibimobileDeviceArguments = "-n -u " + CurrentUDID + " -i " + IPAPath;
+				}
+
 			}
+			return true;
+		}
+
+		public static bool UninstallIPAOnDevices(string IPAPath)
+		{
+			return ActionOnIPAOnDevices(" --uninstall ", IPAPath);
+		}
+
+		public static bool InstallIPAOnDevices(string IPAPath)
+		{
+			return ActionOnIPAOnDevices(" -i ", IPAPath);
+		}
+
+		public static ConnectedDeviceInformation[] GetAllConnectedDevices()
+		{
+			List<ConnectedDeviceInformation> DevicesToReturn = new List<ConnectedDeviceInformation>();
+			string FullOutput = ExecuteLibimobileProcess("idevice_id");
+			string DeviceName = "";
+			string UDID = "";
+			string DeviceType = "";
+			string ConnectingInterface = "";
+
+			var ConnectedDevicesUDIDs = FullOutput.Split(new string[] { Environment.NewLine }, StringSplitOptions.None);
+
+			foreach (var ConnnectedUDID in ConnectedDevicesUDIDs)
+			{
+				ConnectedDeviceInformation NewConnectedDevice;
+				NewConnectedDevice.UDID = "";
+				String ParsedUDID = ConnnectedUDID.Split(' ').First();
+				String IdeviceInfoArgs = "-u " + ParsedUDID;
+
+				if (ConnnectedUDID.Contains("Network"))
+				{
+					ConnectingInterface = "Network";
+					IdeviceInfoArgs = "-n " + IdeviceInfoArgs;
+				}
+				else
+				{
+					ConnectingInterface = "USB";
+				}
 
-			return NewProcess;
+				string OutputInfo = ExecuteLibimobileProcess("ideviceinfo", IdeviceInfoArgs);
+				if (OutputInfo.Contains("not found!"))
+				{
+					Program.LogVerbose(OutputInfo);
+				}
+				else
+				{
+					foreach (string Line in OutputInfo.Split(Environment.NewLine.ToCharArray()))
+					{
+						if (Line.StartsWith("DeviceName: "))
+						{
+							DeviceName = Line.Split(':').Last();
+							DeviceName.Replace(" ", "");
+						}
+						else if (Line.StartsWith("UniqueDeviceID: "))
+						{
+							UDID = Line.Split(':').Last();
+							UDID.Replace(" ", "");
+						}
+						else if (Line.StartsWith("ProductType: "))
+						{
+							DeviceType = Line.Split(':').Last();
+							DeviceType.Replace(" ", "");
+						}
+					}
+				}
+				DevicesToReturn.Add(new ConnectedDeviceInformation(DeviceName, UDID, DeviceType, ConnectingInterface));
+				DeviceName = "";
+				UDID = "";
+				DeviceType = "";
+				ConnectingInterface = "";
+			}
+			return DevicesToReturn.ToArray();
 		}
 	}
 }
diff --git a/Engine/Source/Programs/IOS/iPhonePackager/GraphicalResignTool.cs b/Engine/Source/Programs/IOS/iPhonePackager/GraphicalResignTool.cs
--- a/Engine/Source/Programs/IOS/iPhonePackager/GraphicalResignTool.cs
+++ b/Engine/Source/Programs/IOS/iPhonePackager/GraphicalResignTool.cs
@@ -355,11 +355,11 @@ namespace iPhonePackager
 
 		void ScanConnectedDevices()
 		{
-			ConnectedDeviceInfo [] DeviceMap = null;
+			ConnectedDeviceInformation [] DeviceMap = null;
 
 			Program.ProgressDialog.OnBeginBackgroundWork = delegate
 			{
-				DeviceMap = DeploymentHelper.Get().EnumerateConnectedDevices();
+				DeviceMap = DeploymentHelper.GetAllConnectedDevices();
 			};
 
 			if ((Program.ProgressDialog.ShowDialog() == DialogResult.OK) && (DeviceMap != null))
@@ -621,7 +621,7 @@ namespace iPhonePackager
 			{
 				Program.ProgressDialog.OnBeginBackgroundWork = delegate
 				{
-					DeploymentHelper.Get().BackupDocumentsDirectory(BundleID, Config.GetRootBackedUpDocumentsDirectory());
+					DeploymentHelper.BackupDocumentsDirectory(BundleID, Config.GetRootBackedUpDocumentsDirectory());
 				};
 				Program.ProgressDialog.ShowDialog();
 			}
@@ -635,7 +635,7 @@ namespace iPhonePackager
 			{
 				Program.ProgressDialog.OnBeginBackgroundWork = delegate
 				{
-					DeploymentHelper.Get().UninstallIPAOnDevice(BundleID);
+					DeploymentHelper.UninstallIPAOnDevices(BundleID);
 				};
 				Program.ProgressDialog.ShowDialog();
 			}
diff --git a/Engine/Source/Programs/IOS/iPhonePackager/Utilities.cs b/Engine/Source/Programs/IOS/iPhonePackager/Utilities.cs
--- a/Engine/Source/Programs/IOS/iPhonePackager/Utilities.cs
+++ b/Engine/Source/Programs/IOS/iPhonePackager/Utilities.cs
@@ -922,7 +922,6 @@ namespace iPhonePackager
 		{
 			if (Environment.OSVersion.Platform == PlatformID.Unix || Environment.OSVersion.Platform == PlatformID.MacOSX)
 			{
-				// From Mono 6, the FriendlyName field cannot be set on Unix so we always parse it out of the Subject.
 				return GetCommonNameFromCert(Cert);
 			}
 			else
diff --git a/Engine/Source/Programs/IOS/iPhonePackager/iPhonePackager.cs b/Engine/Source/Programs/IOS/iPhonePackager/iPhonePackager.cs
--- a/Engine/Source/Programs/IOS/iPhonePackager/iPhonePackager.cs
+++ b/Engine/Source/Programs/IOS/iPhonePackager/iPhonePackager.cs
@@ -720,7 +720,7 @@ namespace iPhonePackager
 
         static void ListDevices()
         {
-            var DeviceList = DeploymentHelper.Get().EnumerateConnectedDevices();
+            var DeviceList = DeploymentHelper.GetAllConnectedDevices();
 
             Console.WriteLine("-------------------------------------------------------");
             Console.WriteLine("List of devices attached");
@@ -1048,13 +1048,6 @@ namespace iPhonePackager
                     Program.ReturnCode = (int)ErrorCodes.Error_Unknown;
                 }
             }
-            finally
-            {
-                if (DeploymentHelper.DeploymentServerProcess != null)
-                {
-                    DeploymentHelper.DeploymentServerProcess.Close();
-                }
-            }
 
             Environment.ExitCode = ReturnCode;
             return (ReturnCode);
