Rebuild iPhonePackager after DeploymentServer removal

diff --git a/Engine/Source/Programs/IOS/iPhonePackager/iPhonePackager.csproj b/Engine/Source/Programs/IOS/iPhonePackager/iPhonePackager.csproj
--- a/Engine/Source/Programs/IOS/iPhonePackager/iPhonePackager.csproj
+++ b/Engine/Source/Programs/IOS/iPhonePackager/iPhonePackager.csproj
@@ -221,12 +221,6 @@
     <Content Include="Resources\YellowCheck.png" />
     <None Include="Resources\GreyCheck.png" />
   </ItemGroup>
-  <ItemGroup>
-    <ProjectReference Include="..\DeploymentInterface\DeploymentInterface.csproj">
-      <Project>{B7C89A17-1CB7-42D7-B879-43297538C193}</Project>
-      <Name>DeploymentInterface</Name>
-    </ProjectReference>
-  </ItemGroup>
   <Import Project="$(MSBuildToolsPath)\Microsoft.CSharp.targets" />
   <!-- To modify your build process, add your task inside one of the targets below and uncomment it. 
        Other similar extension points exist, see Microsoft.Common.targets.
