Remove UDK and UE3 references from the iPhone Packager

diff --git a/Engine/Source/Programs/IOS/iPhonePackager/ConfigureMobileGame.cs b/Engine/Source/Programs/IOS/iPhonePackager/ConfigureMobileGame.cs
--- a/Engine/Source/Programs/IOS/iPhonePackager/ConfigureMobileGame.cs
+++ b/Engine/Source/Programs/IOS/iPhonePackager/ConfigureMobileGame.cs
@@ -51,13 +51,13 @@ namespace iPhonePackager
 
 			if (!Helper.GetString("CFBundleName", out BundleName))
 			{
-				BundleName = "MyUDKGame";
+				BundleName = "MyUnrealGame";
 				Helper.SetString("CFBundleName", BundleName);
 			}
 
 			if (!Helper.GetString("CFBundleDisplayName", out BundleDisplayName))
 			{
-				BundleDisplayName = "UDK Game";
+				BundleDisplayName = "Unreal Game";
 				Helper.SetString("CFBundleDisplayName", BundleDisplayName);
 			}
 
diff --git a/Engine/Source/Programs/IOS/iPhonePackager/FileOperations.cs b/Engine/Source/Programs/IOS/iPhonePackager/FileOperations.cs
--- a/Engine/Source/Programs/IOS/iPhonePackager/FileOperations.cs
+++ b/Engine/Source/Programs/IOS/iPhonePackager/FileOperations.cs
@@ -596,7 +596,7 @@ namespace iPhonePackager
 
 
 			/// <summary>
-			/// Finds the root path in a given IPA (e.g., "Payload/UDKGame.app/")
+			/// Finds the root path in a given IPA (e.g., "Payload/UnrealGame.app/")
 			/// </summary>
 			public static string FindRootPathInIPA(ZipFile Zip)
 			{
diff --git a/Engine/Source/Programs/IOS/iPhonePackager/iPhonePackager.cs b/Engine/Source/Programs/IOS/iPhonePackager/iPhonePackager.cs
--- a/Engine/Source/Programs/IOS/iPhonePackager/iPhonePackager.cs
+++ b/Engine/Source/Programs/IOS/iPhonePackager/iPhonePackager.cs
@@ -815,7 +815,7 @@ namespace iPhonePackager
                     Log("RPC Commands: SetExec, InstallProvision, MakeApp, DeleteIPA, Copy, Kill, Strip, Zip, GenDSYM");
                     Log("");
                     Log("Sample commandlines:");
-                    Log(" ... iPhonePackager Deploy UDKGame Release");
+                    Log(" ... iPhonePackager Deploy UnrealGame Release");
                     Log(" ... iPhonePackager RPC SwordGame Shipping MakeApp");
                     return (int)ErrorCodes.Error_Arguments;
                 }
