Update ue4commandline.txt reference

#ios

diff --git a/Engine/Source/Programs/IOS/iPhonePackager/DeployTime.cs b/Engine/Source/Programs/IOS/iPhonePackager/DeployTime.cs
--- a/Engine/Source/Programs/IOS/iPhonePackager/DeployTime.cs
+++ b/Engine/Source/Programs/IOS/iPhonePackager/DeployTime.cs
@@ -168,7 +168,7 @@ namespace iPhonePackager
 							try
 							{
 								// Compare the commandline embedded to prevent us from any unnecessary writing.
-								byte[] CommandlineBytes = Zip.ReadAllBytes("ue4commandline.txt");
+								byte[] CommandlineBytes = Zip.ReadAllBytes("uecommandline.txt");
 								string ExistingCommandline = Encoding.UTF8.GetString(CommandlineBytes, 0, CommandlineBytes.Length);
 								if (ExistingCommandline != AdditionalCommandline)
 								{
@@ -192,7 +192,7 @@ namespace iPhonePackager
 									// Write out the new commandline.
 									FileOperations.ZipFileSystem WritableZip = new FileOperations.ZipFileSystem(TmpFilePath);
 									byte[] NewCommandline = Encoding.UTF8.GetBytes(ProjectFile + " " + AdditionalCommandline);
-									WritableZip.WriteAllBytes("ue4commandline.txt", NewCommandline);
+									WritableZip.WriteAllBytes("uecommandline.txt", NewCommandline);
 
 									// We need to residn the application after the commandline file has changed.
 									CodeSignatureBuilder CodeSigner = new CodeSignatureBuilder();
