Removing launch image logic from IPP since it's now handled by a storyboard.

diff --git a/Engine/Source/Programs/IOS/iPhonePackager/CompileTime.cs b/Engine/Source/Programs/IOS/iPhonePackager/CompileTime.cs
--- a/Engine/Source/Programs/IOS/iPhonePackager/CompileTime.cs
+++ b/Engine/Source/Programs/IOS/iPhonePackager/CompileTime.cs
@@ -257,41 +257,6 @@ namespace iPhonePackager
 			OutString = OutString.Replace("${EXECUTABLE_NAME}", Program.GameName);
 			OutString = OutString.Replace("${BUNDLE_IDENTIFIER}", Program.GameName.Replace("_", ""));
 
-			// this is a temp way to inject the iphone 6 images without needing to upgrade everyone's plist
-			// eventually we want to generate this based on what the user has set in the project settings
-			string[] IPhoneConfigs =  
-				{ 
-					"Default-IPhone6", "Landscape", "{375, 667}", 
-					"Default-IPhone6", "Portrait", "{375, 667}", 
-					"Default-IPhone6Plus-Landscape", "Landscape", "{414, 736}", 
-					"Default-IPhone6Plus-Portrait", "Portrait", "{414, 736}", 
-					"Default", "Landscape", "{320, 480}",
-					"Default", "Portrait", "{320, 480}",
-					"Default-568h", "Landscape", "{320, 568}",
-					"Default-568h", "Portrait", "{320, 568}",
-					"Default-IPhoneXS-Landscape", "Landscape", "{375, 812}",
-					"Default-IPhoneXS-Portrait", "Portrait", "{375, 812}",
-				};
-
-			StringBuilder NewLaunchImagesString = new StringBuilder("<key>UILaunchImages~iphone</key>\n\t\t<array>\n");
-			for (int ConfigIndex = 0; ConfigIndex < IPhoneConfigs.Length; ConfigIndex += 3)
-			{
-				NewLaunchImagesString.Append("\t\t\t<dict>\n");
-				NewLaunchImagesString.Append("\t\t\t\t<key>UILaunchImageMinimumOSVersion</key>\n");
-				NewLaunchImagesString.Append("\t\t\t\t<string>8.0</string>\n");
-				NewLaunchImagesString.Append("\t\t\t\t<key>UILaunchImageName</key>\n");
-				NewLaunchImagesString.AppendFormat("\t\t\t\t<string>{0}</string>\n", IPhoneConfigs[ConfigIndex + 0]);
-				NewLaunchImagesString.Append("\t\t\t\t<key>UILaunchImageOrientation</key>\n");
-				NewLaunchImagesString.AppendFormat("\t\t\t\t<string>{0}</string>\n", IPhoneConfigs[ConfigIndex + 1]);
-				NewLaunchImagesString.Append("\t\t\t\t<key>UILaunchImageSize</key>\n");
-				NewLaunchImagesString.AppendFormat("\t\t\t\t<string>{0}</string>\n", IPhoneConfigs[ConfigIndex + 2]);
-				NewLaunchImagesString.Append("\t\t\t</dict>\n");
-			}
-
-			// close it out
-			NewLaunchImagesString.Append("\t\t\t</array>\n\t\t<key>UILaunchImages~ipad</key>");
-			OutString = OutString.Replace("<key>UILaunchImages~ipad</key>", NewLaunchImagesString.ToString());
-
 			byte[] RawInfoPList = Encoding.UTF8.GetBytes(OutString);
 			File.WriteAllBytes(TargetPListFilename, RawInfoPList);
 
