Make IPP look for mobile provisions in Restricted folders
#rnx

diff --git a/Engine/Source/Programs/IOS/iPhonePackager/Config.cs b/Engine/Source/Programs/IOS/iPhonePackager/Config.cs
--- a/Engine/Source/Programs/IOS/iPhonePackager/Config.cs
+++ b/Engine/Source/Programs/IOS/iPhonePackager/Config.cs
@@ -162,6 +162,14 @@ namespace iPhonePackager
 			get { return Path.GetFullPath(RootRelativePath + @"Engine\Config"); }
 		}
 
+		/// <summary>
+		/// The engine directory
+		/// </summary>
+		public static string EngineDirectory
+		{
+			get { return Path.GetFullPath(RootRelativePath + @"Engine"); }
+		}
+
 		/// <summary>
 		/// The local directory that a payload (GameName.app) is assembled into before being copied to the Mac (on PC)
 		/// </summary>
diff --git a/Engine/Source/Programs/IOS/iPhonePackager/MobileProvisionUtilities.cs b/Engine/Source/Programs/IOS/iPhonePackager/MobileProvisionUtilities.cs
--- a/Engine/Source/Programs/IOS/iPhonePackager/MobileProvisionUtilities.cs
+++ b/Engine/Source/Programs/IOS/iPhonePackager/MobileProvisionUtilities.cs
@@ -45,68 +45,62 @@ namespace iPhonePackager
                 Directory.CreateDirectory(LocalProvisionFolder);
             }
 
-            // copy all of the provisions from the game directory to the library
-            if (!String.IsNullOrEmpty(Config.ProjectFile))
-            {
-                var ProjectFileBuildIOSPath = Path.GetDirectoryName(Config.ProjectFile) + "/Build/" + Config.OSString + "/";
-                Program.Log("Finding provisions in {0}", ProjectFileBuildIOSPath);
-                if (Directory.Exists(ProjectFileBuildIOSPath))
-                {
-                    foreach (string Provision in Directory.EnumerateFiles(ProjectFileBuildIOSPath, "*.mobileprovision", SearchOption.AllDirectories))
-                    {
-                        Program.Log(Provision);
-                        string TargetFile = Config.ProvisionDirectory + Path.GetFileName(Provision);
-                        if (!File.Exists(TargetFile) || File.GetLastWriteTime(TargetFile) < File.GetLastWriteTime(Provision))
-                        {
-                            FileInfo DestFileInfo;
-                            if (File.Exists(TargetFile))
-                            {
-                                DestFileInfo = new FileInfo(TargetFile);
-                                DestFileInfo.Attributes = DestFileInfo.Attributes & ~FileAttributes.ReadOnly;
-                            }
-                            Program.Log("Copying {0} -> {1}", Provision, TargetFile);
-                            File.Copy(Provision, TargetFile, true);
-                            DestFileInfo = new FileInfo(TargetFile);
-                            DestFileInfo.Attributes = DestFileInfo.Attributes & ~FileAttributes.ReadOnly;
-                            if (!File.Exists(TargetFile))
-                            {
-                                Program.Log("ERROR: Failed to copy {0} -> {1}", Provision, TargetFile);
-                            }
-                        }
-                    }
-                }
-            }
-
-            // copy all of the provisions from the engine directory to the library
-            {
-                string ProvisionDirectory = Environment.GetEnvironmentVariable("ProvisionDirectory") ?? Config.EngineBuildDirectory;
-                Program.Log("Finding provisions in {0}", ProvisionDirectory);
-                if (Directory.Exists(ProvisionDirectory))
-                {
-                    foreach (string Provision in Directory.EnumerateFiles(ProvisionDirectory, "*.mobileprovision", SearchOption.AllDirectories))
-                    {
-                        Program.Log(Provision);
-                        string TargetFile = Config.ProvisionDirectory + Path.GetFileName(Provision);
-                        if (!File.Exists(TargetFile) || File.GetLastWriteTime(TargetFile) < File.GetLastWriteTime(Provision))
-                        {
-                            FileInfo DestFileInfo;
-                            if (File.Exists(TargetFile))
-                            {
-                                DestFileInfo = new FileInfo(TargetFile);
-                                DestFileInfo.Attributes = DestFileInfo.Attributes & ~FileAttributes.ReadOnly;
-                            }
-                            Program.Log("Copying {0} -> {1}", Provision, TargetFile);
-                            File.Copy(Provision, Config.ProvisionDirectory + Path.GetFileName(Provision), true);
-                            DestFileInfo = new FileInfo(TargetFile);
-                            DestFileInfo.Attributes = DestFileInfo.Attributes & ~FileAttributes.ReadOnly;
-                            if(!File.Exists(TargetFile))
-                            {
-                                Program.Log("ERROR: Failed to copy {0} -> {1}", Provision, TargetFile);
-                            }
-                        }
-                    }
-                }
-            }
+			List<string> ProvisionCopySrcDirectories = new List<string>();
+
+			// Paths for provisions under game directory 
+			if (!String.IsNullOrEmpty(Config.ProjectFile))
+			{
+				ProvisionCopySrcDirectories.Add(Path.GetDirectoryName(Config.ProjectFile) + "/Build/" + Config.OSString + "/");
+				ProvisionCopySrcDirectories.Add(Path.GetDirectoryName(Config.ProjectFile) + "/Restricted/NoRedist/Build/" + Config.OSString + "/");
+				ProvisionCopySrcDirectories.Add(Path.GetDirectoryName(Config.ProjectFile) + "/Restricted/NotForLicensees/Build/" + Config.OSString + "/");
+			}
+
+			// Paths for provisions under the Engine directory
+			string OverrideProvisionDirectory = Environment.GetEnvironmentVariable("ProvisionDirectory");
+			if(!String.IsNullOrEmpty(OverrideProvisionDirectory))
+			{
+				ProvisionCopySrcDirectories.Add(OverrideProvisionDirectory);
+			}
+			else
+			{
+				ProvisionCopySrcDirectories.Add(Config.EngineBuildDirectory);
+				ProvisionCopySrcDirectories.Add(Config.EngineDirectory + "/Restricted/NoRedist/Build/" + Config.OSString + "/");
+				ProvisionCopySrcDirectories.Add(Config.EngineDirectory + "/Restricted/NotForLicensees/Build/" + Config.OSString + "/");
+			}
+
+			// Copy all provisions from the above paths to the library
+			foreach (string ProvisionCopySrcDirectory in ProvisionCopySrcDirectories)
+			{
+				if (Directory.Exists(ProvisionCopySrcDirectory))
+				{
+					Program.Log("Finding provisions in {0}", ProvisionCopySrcDirectory);
+					foreach (string Provision in Directory.EnumerateFiles(ProvisionCopySrcDirectory, "*.mobileprovision", SearchOption.AllDirectories))
+					{
+						string TargetFile = Config.ProvisionDirectory + Path.GetFileName(Provision);
+						if (!File.Exists(TargetFile) || File.GetLastWriteTime(TargetFile) < File.GetLastWriteTime(Provision))
+						{
+							FileInfo DestFileInfo;
+							if (File.Exists(TargetFile))
+							{
+								DestFileInfo = new FileInfo(TargetFile);
+								DestFileInfo.Attributes = DestFileInfo.Attributes & ~FileAttributes.ReadOnly;
+							}
+							Program.Log("  Copying {0} -> {1}", Provision, TargetFile);
+							File.Copy(Provision, TargetFile, true);
+							DestFileInfo = new FileInfo(TargetFile);
+							DestFileInfo.Attributes = DestFileInfo.Attributes & ~FileAttributes.ReadOnly;
+							if (!File.Exists(TargetFile))
+							{
+								Program.Log("ERROR: Failed to copy {0} -> {1}", Provision, TargetFile);
+							}
+						}
+						else
+						{
+							Program.Log("  Not copying {0} as {1} already exists and is not older", Provision, TargetFile);
+						}
+					}
+				}
+			}
         }
 
         public static string FindCompatibleProvision(string CFBundleIdentifier, out bool bNameMatch, bool bCheckCert = true, bool bCheckIdentifier = true, bool bCheckDistro = true)
