Updated iPhonePackager signature format to support iOS 14.5

diff --git a/Engine/Source/Programs/IOS/iPhonePackager/CodeSigning.cs b/Engine/Source/Programs/IOS/iPhonePackager/CodeSigning.cs
--- a/Engine/Source/Programs/IOS/iPhonePackager/CodeSigning.cs
+++ b/Engine/Source/Programs/IOS/iPhonePackager/CodeSigning.cs
@@ -865,13 +865,8 @@ namespace iPhonePackager
 				EntitlementsBlob FinalEntitlementsBlob = EntitlementsBlob.Create(EntitlementsText);
 
 				// Create the code directory blob
-				uint Version = CodeDirectoryBlob.cVersion2; 
-				if (CodeSigningBlobLC != null)
-				{
-					CodeDirectoryBlob OldCodeDir = CodeSigningBlobLC.Payload.GetBlobByMagic(AbstractBlob.CSMAGIC_CODEDIRECTORY) as CodeDirectoryBlob;
-					Version = OldCodeDir.Version;
-				}
-				CodeDirectoryBlob FinalCodeDirectoryBlob = CodeDirectoryBlob.Create(CFBundleIdentifier, TeamIdentifier, SignedFileLength, Version);
+				uint Version = CodeDirectoryBlob.cVersion3; 
+				CodeDirectoryBlob FinalCodeDirectoryBlob = CodeDirectoryBlob.Create(CFBundleIdentifier, TeamIdentifier, SignedFileLength);
 
 				// Create or preserve the requirements blob
 				RequirementsBlob FinalRequirementsBlob = null;
diff --git a/Engine/Source/Programs/IOS/iPhonePackager/MachObjects.cs b/Engine/Source/Programs/IOS/iPhonePackager/MachObjects.cs
--- a/Engine/Source/Programs/IOS/iPhonePackager/MachObjects.cs
+++ b/Engine/Source/Programs/IOS/iPhonePackager/MachObjects.cs
@@ -2097,6 +2097,7 @@ namespace MachObjectHandling
 	public class CodeDirectoryBlob : AbstractBlob
 	{
 		public const UInt32 cVersion2 = 0x20200;
+		public const UInt32 cVersion3 = 0x20400;
 
 		public UInt32 Version;
 		UInt32 Flags;
@@ -2119,6 +2120,13 @@ namespace MachObjectHandling
 
 		string Team;
 
+		UInt32 Spare3;
+		UInt64 CodeLimit64;
+
+		UInt64 ExecSegBase;
+		UInt64 ExecSegLimit;
+		UInt64 ExecSegFlags;
+
 		// Special slot index for Info.plist
 		public const int cdInfoSlot = 1;
 
@@ -2159,28 +2167,21 @@ namespace MachObjectHandling
 			LogPageSize = SR.ReadByte();
 			Spare2 = SR.ReadUInt32();
 			ScatterCount = SR.ReadUInt32();
-
-			if (Version == cVersion2)
-			{
-				UInt32 TeamStringOffset = SR.ReadUInt32();
-
-				// Read the identifier string
-				SR.PushPositionAndJump(StartOfBlob + IdentifierStringOffset);
-				Identifier = SR.ReadASCIIZ();
-				SR.PopPosition();
-
-				// Read the team string
-				SR.PushPositionAndJump(StartOfBlob + TeamStringOffset);
-				Team = SR.ReadASCIIZ();
-				SR.PopPosition();
-			}
-			else
-			{
-				// Read the identifier string
-				SR.PushPositionAndJump(StartOfBlob + IdentifierStringOffset);
-				Identifier = SR.ReadASCIIZ();
-				SR.PopPosition();
-			}
+			UInt32 TeamStringOffset = SR.ReadUInt32();
+			Spare3 = SR.ReadUInt32();
+			CodeLimit64 = SR.ReadUInt64();
+			ExecSegBase = SR.ReadUInt64();
+			ExecSegLimit = SR.ReadUInt64();
+			ExecSegFlags = SR.ReadUInt64();
+
+			// Read the identifier string
+			SR.PushPositionAndJump(StartOfBlob + IdentifierStringOffset);
+			Identifier = SR.ReadASCIIZ();
+			SR.PopPosition();
+			// Read the team string
+			SR.PushPositionAndJump(StartOfBlob + TeamStringOffset);
+			Team = SR.ReadASCIIZ();
+			SR.PopPosition();
 
 			// Read the hashes
 			long TotalNumHashes = SpecialSlotCount + CodeSlotCount;
@@ -2233,28 +2234,24 @@ namespace MachObjectHandling
 			SW.Write(Spare2);
 			SW.Write(ScatterCount);
 
-			if (Version == cVersion2)
-			{
-				OffsetFieldU32or64 TeamStringOffset = SW.WriteDeferredOffsetFrom(StartPos, Bits.Num._32);
+			OffsetFieldU32or64 TeamStringOffset = SW.WriteDeferredOffsetFrom(StartPos, Bits.Num._32);
+			SW.Write(Spare3);
 
-				// Write the identifier
-				SW.CommitDeferredField(IdentifierStringOffset);
-				byte[] IdentifierOutput = Utilities.CreateASCIIZ(Identifier);
-				SW.Write(IdentifierOutput);
+			SW.Write(CodeLimit64);
+			SW.Write(ExecSegBase);
+			SW.Write(ExecSegLimit);
+			SW.Write(ExecSegFlags);
 
-				// Write the team identifier
-				SW.CommitDeferredField(TeamStringOffset);
-				byte[] TeamOutput = Utilities.CreateASCIIZ(Team);
-				SW.Write(TeamOutput);
-			}
-			else
-			{
-				// Write the identifier
-				SW.CommitDeferredField(IdentifierStringOffset);
-				byte[] IdentifierOutput = Utilities.CreateASCIIZ(Identifier);
-				SW.Write(IdentifierOutput);
-			}
+			// Write the identifier
+			SW.CommitDeferredField(IdentifierStringOffset);
+			byte[] IdentifierOutput = Utilities.CreateASCIIZ(Identifier);
+			SW.Write(IdentifierOutput);
 
+			// Write the team identifier
+			SW.CommitDeferredField(TeamStringOffset);
+			byte[] TeamOutput = Utilities.CreateASCIIZ(Team);
+			SW.Write(TeamOutput);
+		
 			// Write the hashes
 			SW.CommitDeferredField(HashOffset);
 			SW.Write(Hashes);
@@ -2298,7 +2295,7 @@ namespace MachObjectHandling
 			}
 		}
 
-		public static CodeDirectoryBlob Create(string ApplicationID, string TeamID, int SignedFileLength, uint Version = cVersion2)
+		public static CodeDirectoryBlob Create(string ApplicationID, string TeamID, int SignedFileLength)
 		{
 			CodeDirectoryBlob Blob = new CodeDirectoryBlob();
 			Blob.Allocate(ApplicationID, TeamID, SignedFileLength);
@@ -2306,17 +2303,24 @@ namespace MachObjectHandling
 			return Blob;
 		}
 
-		public void Allocate(string ApplicationID, string TeamID, int SignedFileLength, uint InVersion = cVersion2)
+		public void Allocate(string ApplicationID, string TeamID, int SignedFileLength)
 		{
 			Identifier = ApplicationID;
 			Team = TeamID;
 
-			Version = InVersion;
+			Version = cVersion3;
 			Flags = 0;
 			Spare1 = 0;
 			Spare2 = 0;
+			Spare3 = 0;
 			ScatterCount = 0;
 
+			CodeLimit64 = 0;
+
+			ExecSegBase = 0;
+			ExecSegLimit = 0;
+			ExecSegFlags = 0;
+
 			// 4 KB pages
 			LogPageSize = 12;
 			int PageSize = 1 << LogPageSize;
