Fix warning in resigning tool on iOS 14.5

diff --git a/Engine/Source/Programs/IOS/iPhonePackager/CodeSigning.cs b/Engine/Source/Programs/IOS/iPhonePackager/CodeSigning.cs
--- a/Engine/Source/Programs/IOS/iPhonePackager/CodeSigning.cs
+++ b/Engine/Source/Programs/IOS/iPhonePackager/CodeSigning.cs
@@ -865,7 +865,6 @@ namespace iPhonePackager
 				EntitlementsBlob FinalEntitlementsBlob = EntitlementsBlob.Create(EntitlementsText);
 
 				// Create the code directory blob
-				uint Version = CodeDirectoryBlob.cVersion3; 
 				CodeDirectoryBlob FinalCodeDirectoryBlob = CodeDirectoryBlob.Create(CFBundleIdentifier, TeamIdentifier, SignedFileLength);
 
 				// Create or preserve the requirements blob
