Fixed bug in iPhonePackager Mach-O handling, updated signatures to use SHA-256.

Taking UE4 version of executable over UE5 per Rolando

diff --git a/Engine/Source/Programs/IOS/iPhonePackager/CodeSigning.cs b/Engine/Source/Programs/IOS/iPhonePackager/CodeSigning.cs
--- a/Engine/Source/Programs/IOS/iPhonePackager/CodeSigning.cs
+++ b/Engine/Source/Programs/IOS/iPhonePackager/CodeSigning.cs
@@ -821,6 +821,7 @@ namespace iPhonePackager
 				// Find out if there was an existing code sign blob and find the linkedit segment command
 				MachLoadCommandCodeSignature CodeSigningBlobLC = null;
 				MachLoadCommandSegment LinkEditSegmentLC = null;
+				MachLoadCommandSegment TextSegmentLC = null;
 				foreach (MachLoadCommand Command in Exe.Commands)
 				{
 					if (CodeSigningBlobLC == null)
@@ -836,6 +837,15 @@ namespace iPhonePackager
 							LinkEditSegmentLC = null;
 						}
 					}
+
+					if (TextSegmentLC == null)
+					{
+						TextSegmentLC = Command as MachLoadCommandSegment;
+						if (TextSegmentLC.SegmentName != "__TEXT")
+						{
+							TextSegmentLC = null;
+						}
+					}
 				}
 
 				if (LinkEditSegmentLC == null)
@@ -859,13 +869,22 @@ namespace iPhonePackager
 
 				int SignedFileLength = (int)CodeSigningBlobLC.BlobFileOffset;
 
+				// Set ExecSegment data
+				UInt64 ExecSegBase = 0;
+				UInt64 ExecSegLimit = 0;
+				if (TextSegmentLC != null)
+				{
+					ExecSegBase = TextSegmentLC.FileOffset;
+					ExecSegLimit = TextSegmentLC.FileSize;
+				}
+
 				// Create the entitlements blob
 				string TeamIdentifier = Provision.ApplicationIdentifierPrefix;
 				string EntitlementsText = BuildEntitlementString(CFBundleIdentifier, out TeamIdentifier);
 				EntitlementsBlob FinalEntitlementsBlob = EntitlementsBlob.Create(EntitlementsText);
 
 				// Create the code directory blob
-				CodeDirectoryBlob FinalCodeDirectoryBlob = CodeDirectoryBlob.Create(CFBundleIdentifier, TeamIdentifier, SignedFileLength);
+				CodeDirectoryBlob FinalCodeDirectoryBlob = CodeDirectoryBlob.Create(CFBundleIdentifier, TeamIdentifier, SignedFileLength, ExecSegBase, ExecSegLimit);
 
 				// Create or preserve the requirements blob
 				RequirementsBlob FinalRequirementsBlob = null;
diff --git a/Engine/Source/Programs/IOS/iPhonePackager/MachObjects.cs b/Engine/Source/Programs/IOS/iPhonePackager/MachObjects.cs
--- a/Engine/Source/Programs/IOS/iPhonePackager/MachObjects.cs
+++ b/Engine/Source/Programs/IOS/iPhonePackager/MachObjects.cs
@@ -435,8 +435,8 @@ namespace MachObjectHandling
 			else
 			{
 				byte[] Buffer = BitConverter.GetBytes(Value);
-				Array.Reverse(Buffer, 0, 4);
-				SW.Write(Buffer, 0, 4);
+				Array.Reverse(Buffer, 0, 8);
+				SW.Write(Buffer, 0, 8);
 			}
 		}
 
@@ -2148,12 +2148,13 @@ namespace MachObjectHandling
 		/// <summary>
 		/// Hash provider (SHA1)
 		/// </summary>
-		protected SHA1CryptoServiceProvider HashProvider = new SHA1CryptoServiceProvider();
+		protected SHA256CryptoServiceProvider HashProvider2 = new SHA256CryptoServiceProvider();
 
 		protected override void UnpackageData(ReadingContext SR, UInt32 Length)
 		{
 			long StartOfBlob = SR.Position - sizeof(UInt32) * 2;
 
+			// Reference: https://llvm.org/doxygen/BinaryFormat_2MachO_8h_source.html
 			Version = SR.ReadUInt32();
 			Flags = SR.ReadUInt32();
 			UInt32 HashOffset = SR.ReadUInt32();
@@ -2275,7 +2276,7 @@ namespace MachObjectHandling
 		/// </summary>
 		public void GenerateSpecialSlotHash(int SpecialSlotIndex, byte[] SourceData)
 		{
-			byte[] Hash = HashProvider.ComputeHash(SourceData);
+			byte[] Hash = HashProvider2.ComputeHash(SourceData);
 			Array.Copy(
 				Hash, 0,
 				Hashes, (SpecialSlotCount - SpecialSlotIndex) * BytesPerHash,
@@ -2295,15 +2296,15 @@ namespace MachObjectHandling
 			}
 		}
 
-		public static CodeDirectoryBlob Create(string ApplicationID, string TeamID, int SignedFileLength)
+		public static CodeDirectoryBlob Create(string ApplicationID, string TeamID, int SignedFileLength, UInt64 InExecSegBase, UInt64 InExecSegLimit)
 		{
 			CodeDirectoryBlob Blob = new CodeDirectoryBlob();
-			Blob.Allocate(ApplicationID, TeamID, SignedFileLength);
+			Blob.Allocate(ApplicationID, TeamID, SignedFileLength, InExecSegBase, InExecSegLimit);
 
 			return Blob;
 		}
 
-		public void Allocate(string ApplicationID, string TeamID, int SignedFileLength)
+		public void Allocate(string ApplicationID, string TeamID, int SignedFileLength, UInt64 InExecSegBase, UInt64 InExecSegLimit)
 		{
 			Identifier = ApplicationID;
 			Team = TeamID;
@@ -2317,18 +2318,18 @@ namespace MachObjectHandling
 
 			CodeLimit64 = 0;
 
-			ExecSegBase = 0;
-			ExecSegLimit = 0;
-			ExecSegFlags = 0;
+			ExecSegBase = InExecSegBase;
+			ExecSegLimit = InExecSegLimit;
+			ExecSegFlags = 17;
 
 			// 4 KB pages
 			LogPageSize = 12;
 			int PageSize = 1 << LogPageSize;
 
-			// 20 byte SHA1 hashes
-			HashType = 1;
-			BytesPerHash = (byte)(HashProvider.HashSize / 8);
-			Debug.Assert(BytesPerHash == 20);
+			// 32 byte SHA256 hashes
+			HashType = 2;
+			BytesPerHash = (byte)(HashProvider2.HashSize / 8);
+			Debug.Assert(BytesPerHash == 32);
 
 			// Allocate space for the hashes
 			MainImageSignatureLimit = (UInt32)SignedFileLength;
@@ -2347,7 +2348,7 @@ namespace MachObjectHandling
 			{
 				int StartOffset = i * PageSize;
 				int LengthRemaining = (int)MainImageSignatureLimit - StartOffset;
-				byte[] PageHash = HashProvider.ComputeHash(SignedFileData, StartOffset, Math.Min(LengthRemaining, PageSize)); 
+				byte[] PageHash = HashProvider2.ComputeHash(SignedFileData, StartOffset, Math.Min(LengthRemaining, PageSize)); 
 				Array.Copy(PageHash, 0, Hashes, (SpecialSlotCount + i) * BytesPerHash, BytesPerHash);
 			}
 		}
