Fixing iPhonePackager compile failing on new machines.

Copy Local needs to be off, to allow UAT to compile that tool correctly.

#rnx

diff --git a/Engine/Source/Programs/IOS/iPhonePackager/iPhonePackager.csproj b/Engine/Source/Programs/IOS/iPhonePackager/iPhonePackager.csproj
--- a/Engine/Source/Programs/IOS/iPhonePackager/iPhonePackager.csproj
+++ b/Engine/Source/Programs/IOS/iPhonePackager/iPhonePackager.csproj
@@ -78,6 +78,7 @@
     <Reference Include="SysadminsLV.Asn1Parser, Version=1.2.1.0, Culture=neutral, processorArchitecture=MSIL">
       <SpecificVersion>False</SpecificVersion>
       <HintPath>..\..\..\..\Binaries\ThirdParty\IOS\SysadminsLV.Asn1Parser.dll</HintPath>
+      <Private>False</Private>
     </Reference>
     <Reference Include="System" />
     <Reference Include="System.Core">
