Fix idevicefs pull call in IPP

diff --git a/Engine/Source/Programs/IOS/iPhonePackager/DeployTime.cs b/Engine/Source/Programs/IOS/iPhonePackager/DeployTime.cs
--- a/Engine/Source/Programs/IOS/iPhonePackager/DeployTime.cs
+++ b/Engine/Source/Programs/IOS/iPhonePackager/DeployTime.cs
@@ -402,7 +402,14 @@ namespace iPhonePackager
 				// Source folder
 				string SourceFolder = "/Documents/ ";
 
-				string Command = " pull " + SourceFolder + "\"" +  TargetFolder + "\"";
+				bool FolderExists = System.IO.Directory.Exists(TargetFolder);
+
+				if (!FolderExists)
+				{
+					System.IO.Directory.CreateDirectory(TargetFolder);
+				}
+
+				string Command = " pull " + SourceFolder + TargetFolder;
 
 				FullOutput = ExecuteLibimobileProcess("idevicefs", "-u " + Device.UDID + " -b " + BundleIdentifier + Command);
 				if (FullOutput != "0")
@@ -506,8 +513,7 @@ namespace iPhonePackager
 					{
 						if (Line.StartsWith("DeviceName: "))
 						{
-							DeviceName = Line.Split(':').Last();
-							DeviceName.Replace(" ", "");
+							DeviceName = Line.Split(' ').Last();
 						}
 						else if (Line.StartsWith("UniqueDeviceID: "))
 						{
