Fix unicode device name handling issues in IPhonePackager causing issues downloading log files

diff --git a/Engine/Source/Programs/IOS/iPhonePackager/DeployTime.cs b/Engine/Source/Programs/IOS/iPhonePackager/DeployTime.cs
--- a/Engine/Source/Programs/IOS/iPhonePackager/DeployTime.cs
+++ b/Engine/Source/Programs/IOS/iPhonePackager/DeployTime.cs
@@ -251,7 +251,7 @@ namespace iPhonePackager
 		}
 
 		static DeployTimeReporter Reporter = new DeployTimeReporter();
-		static string ExecuteLibimobileProcess(string ProcessName, string ProcessArguments = "")
+		static string ExecuteLibimobileProcess(string ProcessName, string ProcessArguments = "", bool bTreatOutputAsUTF8 = true)
 		{
 			string ExePath = Directory.GetCurrentDirectory() + "/../../../Extras/ThirdPartyNotUE/libimobiledevice/";
 
@@ -275,7 +275,11 @@ namespace iPhonePackager
 			StartInfo.RedirectStandardOutput = true;
 			StartInfo.RedirectStandardError = true;
 			StartInfo.CreateNoWindow = true;
-			StartInfo.StandardOutputEncoding = Encoding.UTF8;
+			if(bTreatOutputAsUTF8)
+			{
+				StartInfo.StandardOutputEncoding = Encoding.UTF8;
+				StartInfo.StandardErrorEncoding = Encoding.UTF8;
+			}
 
 			string FullOutput = "";
 			string ErrorOutput = "";
@@ -345,7 +349,7 @@ namespace iPhonePackager
 
 				IdeviceFSArgs = IdeviceFSArgs + " -u " + Device.UDID + " -b " + BundleIdentifier + " -x \"" + Directory.GetCurrentDirectory() + "/CommandsToPush.txt" + "\"";
 
-				FullOutput = ExecuteLibimobileProcess("idevicefs", IdeviceFSArgs);
+				FullOutput = ExecuteLibimobileProcess("idevicefs", IdeviceFSArgs, false);
 				File.Delete(Directory.GetCurrentDirectory() + "/CommandsToPush.txt");
 
 				if (FullOutput != "0")
@@ -375,7 +379,7 @@ namespace iPhonePackager
 				}
 				System.IO.File.WriteAllText(Directory.GetCurrentDirectory() + "/CommandsToPull.txt", AllCommands);
 				string arguments = "\"" + "-u " + Device.UDID + " -b " + BundleIdentifier + " -x " + Directory.GetCurrentDirectory() + "/CommandsToPull.txt" + "\"";
-				FullOutput = ExecuteLibimobileProcess("idevicefs", arguments);
+				FullOutput = ExecuteLibimobileProcess("idevicefs", arguments, false);
 			}
 			File.Delete(Directory.GetCurrentDirectory() + "/CommandsToPull.txt");
 			if (FullOutput != "0")
@@ -398,27 +402,24 @@ namespace iPhonePackager
 
 				// Destination folder
 				string TargetFolder = Path.Combine(DestinationDocumentsDirectory, Device.DeviceName);
-
-				// Source folder
-				string SourceFolder = "/Documents/ ";
-
-				bool FolderExists = System.IO.Directory.Exists(TargetFolder);
-
-				if (!FolderExists)
+				if (!System.IO.Directory.Exists(TargetFolder))
 				{
 					System.IO.Directory.CreateDirectory(TargetFolder);
 				}
 
-				string Command = " pull " + SourceFolder + TargetFolder;
-
-				FullOutput = ExecuteLibimobileProcess("idevicefs", "-u " + Device.UDID + " -b " + BundleIdentifier + Command);
+				// Pull /Documents
+				string SourceFolder = "/Documents/ ";
+				string Command = " pull " + SourceFolder + "\"" + TargetFolder + "\"";
+				FullOutput = ExecuteLibimobileProcess("idevicefs", "-u " + Device.UDID + " -b " + BundleIdentifier + Command, false);
 				if (FullOutput != "0")
 				{
 					Program.LogVerbose(FullOutput);
 				}
-				SourceFolder = "/Library/Caches/";
-				Command = " pull " + SourceFolder;
-				FullOutput = ExecuteLibimobileProcess("idevicefs", "-u " + Device.UDID + " -b " + BundleIdentifier + Command);
+
+				// Pull /Library/Caches
+				SourceFolder = "/Library/Caches/ ";
+				Command = " pull " + SourceFolder + "\"" + TargetFolder + "\"";
+				FullOutput = ExecuteLibimobileProcess("idevicefs", "-u " + Device.UDID + " -b " + BundleIdentifier + Command, false);
 				if (FullOutput != "0")
 				{
 					Program.LogVerbose(FullOutput);
@@ -513,7 +514,7 @@ namespace iPhonePackager
 					{
 						if (Line.StartsWith("DeviceName: "))
 						{
-							DeviceName = Line.Split(' ').Last();
+							DeviceName = Line.Substring(12);
 						}
 						else if (Line.StartsWith("UniqueDeviceID: "))
 						{
