- Fixed some IPhonePackager codesigning issues with modern xcode (one recoverable error in Development, and some fatal errors in non-Development)

diff --git a/Engine/Source/Programs/IOS/iPhonePackager/Config.cs b/Engine/Source/Programs/IOS/iPhonePackager/Config.cs
--- a/Engine/Source/Programs/IOS/iPhonePackager/Config.cs
+++ b/Engine/Source/Programs/IOS/iPhonePackager/Config.cs
@@ -427,8 +427,10 @@ namespace iPhonePackager
 		/// </summary>
 		public static string AppDirectoryInZIP
 		{
-			get { return String.Format("Payload/{0}{1}.app", GetTargetName(), Program.Architecture); }
+			get { return _AppDirectoryInZIP != null ? _AppDirectoryInZIP : String.Format("Payload/{0}{1}.app", GetTargetName(), Program.Architecture); }
+			set { _AppDirectoryInZIP = value.TrimEnd("/".ToCharArray()); }
 		}
+		private static string _AppDirectoryInZIP = null;
 
 		/// <summary>
 		/// If the requirements blob is present in the existing executable when code signing, should it be carried forward
diff --git a/Engine/Source/Programs/IOS/iPhonePackager/CookTime.cs b/Engine/Source/Programs/IOS/iPhonePackager/CookTime.cs
--- a/Engine/Source/Programs/IOS/iPhonePackager/CookTime.cs
+++ b/Engine/Source/Programs/IOS/iPhonePackager/CookTime.cs
@@ -67,12 +67,14 @@ namespace iPhonePackager
 			// Open up the zip file
 			ZipFile Stub = ZipFile.Read(WorkIPAPath);
 
+			// modern xcode uses Foo-IOS-Config style .app directories, so don't assume GameName, pull it out of of the zip
+			Config.AppDirectoryInZIP = FileOperations.ZipFileSystem.FindRootPathInIPA(Stub);
+
 			// Do a few quick spot checks to catch problems that may have occurred earlier
 			bool bHasCodeSignature = Stub[Config.AppDirectoryInZIP + "/_CodeSignature/CodeResources"] != null;
-			bool bHasMobileProvision = Stub[Config.AppDirectoryInZIP + "/embedded.mobileprovision"] != null;
-			if (!bHasCodeSignature || !bHasMobileProvision)
+			if (!bHasCodeSignature)
 			{
-				Program.Error("Stub IPA does not appear to be signed correctly (missing mobileprovision or CodeResources)");
+				Program.Error("Stub IPA does not appear to be signed correctly (missing _CodeSignature/CodeResources)");
 				Program.ReturnCode = (int)ErrorCodes.Error_StubNotSignedCorrectly;
 			}
 
@@ -223,10 +225,7 @@ namespace iPhonePackager
 				return;
 			}
 
-			string ZipWorkingDir = String.Format("Payload/{0}{1}.app/", Config.GetTargetName(), Program.Architecture);
-
-			FileOperations.ZipFileSystem FileSystem = new FileOperations.ZipFileSystem(Zip, ZipWorkingDir);
-
+			FileOperations.ZipFileSystem FileSystem = new FileOperations.ZipFileSystem(Zip);
 			// Check for a staged plist that needs to be merged into the main one
 			{
 				// Determine if there is a staged one we should try to use instead
@@ -244,7 +243,9 @@ namespace iPhonePackager
 						byte[] StubPListBytes = FileSystem.ReadAllBytes("Info.plist");
 						Utilities.PListHelper StubInfo = new Utilities.PListHelper(Encoding.UTF8.GetString(StubPListBytes));
 
-						StubInfo.MergePlistIn(StageInfoString);
+						// don't overwrite the Exe name in the plist, because we may be merging in Development staged data
+						// with a Shipping binary (or whatever combination), and modern uses different-named executables for each config
+						StubInfo.MergePlistIn(StageInfoString, new HashSet<string> { "CFBundleExecutable" } );
 
 						// Write it back to the cloned stub, where it will be used for all subsequent actions
 						byte[] MergedPListBytes = Encoding.UTF8.GetBytes(StubInfo.SaveToString());
@@ -315,9 +316,8 @@ namespace iPhonePackager
 					string AbsoluteFilename = Path.GetFullPath(Filename);
 					string RelativeFilename = AbsoluteFilename.Substring(SourceDir.Length + 1).Replace('\\', '/');
 
-					string ZipAbsolutePath = String.Format("Payload/{0}{1}.app/{2}",
-						Config.GetTargetName(),
-						Program.Architecture,
+					string ZipAbsolutePath = String.Format("{0}/{1}",
+						Config.AppDirectoryInZIP,
 						RelativeFilename);
 
 					Stream FileContents = File.OpenRead(AbsoluteFilename);
diff --git a/Engine/Source/Programs/IOS/iPhonePackager/Utilities.cs b/Engine/Source/Programs/IOS/iPhonePackager/Utilities.cs
--- a/Engine/Source/Programs/IOS/iPhonePackager/Utilities.cs
+++ b/Engine/Source/Programs/IOS/iPhonePackager/Utilities.cs
@@ -396,7 +396,7 @@ namespace iPhonePackager
 			/// Merge two plists together.  Whenever both have the same key, the value in the dominant source list wins.
 			/// This is special purpose code, and only handles things inside of the <dict> tag
 			/// </summary>
-			public void MergePlistIn(string DominantPlist)
+			public void MergePlistIn(string DominantPlist, HashSet<string> WeakKeysToKeep=null)
 			{
 				if (bReadOnly)
 				{
@@ -422,7 +422,8 @@ namespace iPhonePackager
 						DictionaryNode.AppendChild(Doc.ImportNode(StrongKeyNode, true));
 						DictionaryNode.AppendChild(Doc.ImportNode(StrongKeyNode.NextSibling, true));
 					}
-					else
+					// don't overwrite values we want to keep
+					else if (WeakKeysToKeep == null || !WeakKeysToKeep.Contains(StrongKey))
 					{
 						// Remove the existing value node from the weak file
 						WeakNode.ParentNode.RemoveChild(WeakNode.NextSibling);
