In remote iOS build, ensure that Provision Profiles folder is created prior to installing the Windows specified provisioning profile on the remote mac.
Work around new MacOS 26 codesigning CMS blob size errors by no longer enforcing that replacement signature on Windows be done in place of the one generated on MacOS as the size of these no longer match. This involves restoring the old code that was offsetting the __LINKEDIT section to accomodate changes in the size of the LC_CODE_SIGNATURE section.

diff --git a/Engine/Source/Programs/IOS/iPhonePackager/CodeSigning.cs b/Engine/Source/Programs/IOS/iPhonePackager/CodeSigning.cs
--- a/Engine/Source/Programs/IOS/iPhonePackager/CodeSigning.cs
+++ b/Engine/Source/Programs/IOS/iPhonePackager/CodeSigning.cs
@@ -787,7 +787,10 @@ namespace iPhonePackager
 
 			//@TODO: Verify it's an executable (not an object file, etc...)
 			ulong CurrentStreamOffset = 0;
-			byte[] FinalExeData = new byte[SourceExeData.Length + 1024 * 1024];
+			int OverSize = 1024 * 1024;
+			// Allocate a large additional buffer to account for the fact that signature generated by IPP/MacOS 15 or older
+			// are significantly larger than the cones generated by MacOS 26's codesign
+			byte[] FinalExeData = new byte[SourceExeData.Length + 4 * OverSize];
 			int ArchIndex = 0;
 
 			foreach (MachObjectFile Exe in FatBinary.MachObjectFiles)
@@ -795,7 +798,6 @@ namespace iPhonePackager
 				Program.Log("... Processing one mach object (binary is {0})", FatBinary.bIsFatBinary ? "fat" : "thin");
 
 				// Pad the memory stream with extra room to handle any possible growth in the code signing data
-				int OverSize = 1024 * 1024;
 				int ExeSize = (FatBinary.bIsFatBinary ? (int)FatBinary.Archs[ArchIndex].Size : SourceExeData.Length);
 				MemoryStream OutputExeStream = new MemoryStream(ExeSize + OverSize);
 
@@ -956,9 +958,8 @@ namespace iPhonePackager
 				long NonCodeSigSize = (long)LinkEditSegmentLC.FileSize - CodeSigningBlobLC.BlobFileSize;
 				long BlobStartPosition = NonCodeSigSize + (long)LinkEditSegmentLC.FileOffset;
 
-				// we no longer update the offsets in the load command as we keep the blob size the same
-//				LinkEditSegmentLC.PatchFileLength(OutputExeContext, (uint)(NonCodeSigSize + BlobLength));
-//				CodeSigningBlobLC.PatchPositionAndSize(OutputExeContext, (uint)BlobStartPosition, (uint)BlobLength);
+				LinkEditSegmentLC.PatchFileLength(OutputExeContext, (uint)(NonCodeSigSize + BlobLength));
+				CodeSigningBlobLC.PatchPositionAndSize(OutputExeContext, (uint)BlobStartPosition, (uint)BlobLength);
 
 				// Now that the executable loader command has been inserted and the appropriate section modified, compute all the hashes
 				Program.Log("... Computing hashes ({0:0.00} s elapsed so far)", (DateTime.Now - SigningTime).TotalSeconds);
@@ -1007,10 +1008,6 @@ namespace iPhonePackager
 				Program.Log("... Truncating/copying final binary", DateTime.Now - SigningTime);
 				ulong DesiredExecutableLength = LinkEditSegmentLC.FileSize + LinkEditSegmentLC.FileOffset;
 
-				if ((ulong)Length < DesiredExecutableLength)
-				{
-					throw new InvalidDataException("Data written is smaller than expected, unable to finish signing process");
-				}
 				byte[] Data = OutputExeStream.ToArray();
 				Data.CopyTo(FinalExeData, (long)CurrentStreamOffset);
 				CurrentStreamOffset += DesiredExecutableLength;
diff --git a/Engine/Source/Programs/UnrealBuildTool/Platform/IOS/IOSToolChain.cs b/Engine/Source/Programs/UnrealBuildTool/Platform/IOS/IOSToolChain.cs
--- a/Engine/Source/Programs/UnrealBuildTool/Platform/IOS/IOSToolChain.cs
+++ b/Engine/Source/Programs/UnrealBuildTool/Platform/IOS/IOSToolChain.cs
@@ -1448,6 +1448,7 @@ namespace UnrealBuildTool
 					// Copy the mobile provision into the system store
 					if (Target.ImportProvision != null)
 					{
+						Writer.WriteLine("mkdir -p {0}", Utils.EscapeShellArgument(AppleExports.GetProvisionDirectory().FullName));
 						Writer.WriteLine("cp -f {0} {1}", Utils.EscapeShellArgument(Target.ImportProvision), Utils.EscapeShellArgument(AppleExports.GetProvisionDirectory().FullName));
 					}
 
diff --git a/Engine/Source/Programs/UnrealBuildTool/ToolChain/AppleToolChain.cs b/Engine/Source/Programs/UnrealBuildTool/ToolChain/AppleToolChain.cs
--- a/Engine/Source/Programs/UnrealBuildTool/ToolChain/AppleToolChain.cs
+++ b/Engine/Source/Programs/UnrealBuildTool/ToolChain/AppleToolChain.cs
@@ -1239,6 +1239,7 @@ namespace UnrealBuildTool
 				}
 
 				// copy the provision into standard location
+				Writer.WriteLine("mkdir -p {0}", Utils.EscapeShellArgument(AppleExports.GetProvisionDirectory().FullName));
 				Writer.WriteLine("cp -f {0} {1}", Utils.EscapeShellArgument(Target.RemoteImportProvision), Utils.EscapeShellArgument(AppleExports.GetProvisionDirectory().FullName));
 				MobileProvisionContents MobileProvision = MobileProvisionContents.Read(new FileReference(Target.RemoteImportProvision));
 				MobileProvisionUUID = MobileProvision.GetUniqueId();
