iPhonePackager: when resigning the app and resizing the FileSize of a mach-o load command segment to accomodate different sized code signatures, ensure the the VirtualSize (a.k.a. vmsize) is also recomputed to support the segment's filesize.

#Synced-CL 48480311

diff --git a/Engine/Source/Programs/IOS/iPhonePackager/MachObjects.cs b/Engine/Source/Programs/IOS/iPhonePackager/MachObjects.cs
--- a/Engine/Source/Programs/IOS/iPhonePackager/MachObjects.cs
+++ b/Engine/Source/Programs/IOS/iPhonePackager/MachObjects.cs
@@ -895,7 +895,18 @@ namespace MachObjectHandling
 
 			long PatchOffset = StartingLoadOffset + (2 * sizeof(UInt32)); // Command, CommandLength
 			PatchOffset += 16; // SegmentName
-			PatchOffset += 3 * Bits.Bytes(AddressSize); // VirtualAddress, VirtualSize, FileOffset
+			PatchOffset += 1 * Bits.Bytes(AddressSize); // VirtualAddress
+			// Adjust VirtualSize to accomodate a larger FileSize
+			if (FileSize > VirtualSize)
+			{
+				// VirtualSize is the page aligned memory size required to load FileSize
+				ulong PageAlignmentMask = (AddressSize == Bits.Num._64 ? (ulong)0x4000 : (ulong)0x1000) - 1;
+				VirtualSize = (FileSize + PageAlignmentMask) & ~PageAlignmentMask;
+				SW.PushPositionAndJump(PatchOffset);
+				SW.WriteUInt(VirtualSize, AddressSize);
+				SW.PopPosition();
+			}
+			PatchOffset += 2 * Bits.Bytes(AddressSize); // VirtualSize, FileOffset
 
 			SW.PushPositionAndJump(PatchOffset);
 			SW.WriteUInt(FileSize, AddressSize);
